Http for Blueprint

diff --git a/Engine/Plugins/JsonBlueprintUtilities/Source/JsonBlueprintGraph/Private/K2Node_StructToJsonString.cpp b/Engine/Plugins/JsonBlueprintUtilities/Source/JsonBlueprintGraph/Private/K2Node_StructToJsonString.cpp
new file mode 100644
--- /dev/null
+++ b/Engine/Plugins/JsonBlueprintUtilities/Source/JsonBlueprintGraph/Private/K2Node_StructToJsonString.cpp
@@ -0,0 +1,92 @@
+﻿// Copyright Epic Games, Inc. All Rights Reserved.
+
+#include "K2Node_StructToJsonString.h"
+
+#include "BlueprintActionDatabaseRegistrar.h"
+#include "BlueprintNodeSpawner.h"
+#include "JsonBlueprintFunctionLibrary.h"
+
+#include "Kismet2/BlueprintEditorUtils.h"
+
+#define LOCTEXT_NAMESPACE "K2Node_StructToJsonString"
+
+UK2Node_StructToJsonString::UK2Node_StructToJsonString()
+{
+	FunctionReference.SetExternalMember(
+		GET_FUNCTION_NAME_CHECKED(UJsonBlueprintFunctionLibrary, StructToJsonString),
+		UJsonBlueprintFunctionLibrary::StaticClass());
+}
+
+void UK2Node_StructToJsonString::AllocateDefaultPins()
+{
+	Super::AllocateDefaultPins();
+
+	struct FPinStatics
+	{
+		int32 InputStructPinIndex = -1;
+		int32 OutputStringIndex = -1;
+
+		FPinStatics(const TArray<UEdGraphPin*>& InPins)
+		{
+			InputStructPinIndex = FindPinByName(InPins, TEXT("Struct"));
+			OutputStringIndex = FindPinByName(InPins, TEXT("OutJsonString"));
+		}
+
+		static int32 FindPinByName(const TArray<UEdGraphPin*>& InPins, const FName& InName)
+		{
+			return InPins.IndexOfByPredicate([&InName](const UEdGraphPin* InPin)->bool
+			{
+				return InPin->GetFName() == InName;
+			});
+		}
+		
+	} static PinInfo{Pins};
+
+	UEdGraphPin* InputStructPin = GetPinAt(PinInfo.InputStructPinIndex);
+	check(InputStructPin);
+
+	UEdGraphPin* OutputStringPin = GetPinAt(PinInfo.OutputStringIndex);
+	check(OutputStringPin);
+	
+}
+
+FText UK2Node_StructToJsonString::GetTooltipText() const
+{
+	return LOCTEXT("TooltipText", "Converts a Structure to a raw JSON string.");
+}
+
+bool UK2Node_StructToJsonString::IsNodePure() const
+{
+	return false;
+}
+
+void UK2Node_StructToJsonString::GetMenuActions(FBlueprintActionDatabaseRegistrar& InActionRegistrar) const
+{
+	const UClass* ActionKey = GetClass();
+	if (InActionRegistrar.IsOpenForRegistration(ActionKey))
+	{
+		UBlueprintNodeSpawner* NodeSpawner = UBlueprintNodeSpawner::Create(GetClass());
+		check(NodeSpawner != nullptr);
+
+		InActionRegistrar.AddBlueprintAction(ActionKey, NodeSpawner);
+	}
+}
+
+FText UK2Node_StructToJsonString::GetMenuCategory() const
+{
+	static FText MenuCategory = LOCTEXT("MenuCategory", "Json");
+	return MenuCategory;
+}
+
+void UK2Node_StructToJsonString::NotifyInputChanged() const
+{
+	if (UBlueprint* BP = GetBlueprint())
+	{
+		FBlueprintEditorUtils::MarkBlueprintAsModified(BP);
+	}
+
+	UEdGraph* Graph = GetGraph();
+	Graph->NotifyGraphChanged();
+}
+
+#undef LOCTEXT_NAMESPACE
diff --git a/Engine/Plugins/JsonBlueprintUtilities/Source/JsonBlueprintGraph/Private/K2Node_StructToJsonString.h b/Engine/Plugins/JsonBlueprintUtilities/Source/JsonBlueprintGraph/Private/K2Node_StructToJsonString.h
new file mode 100644
--- /dev/null
+++ b/Engine/Plugins/JsonBlueprintUtilities/Source/JsonBlueprintGraph/Private/K2Node_StructToJsonString.h
@@ -0,0 +1,31 @@
+﻿// Copyright Epic Games, Inc. All Rights Reserved.
+
+#pragma once
+
+#include "CoreMinimal.h"
+#include "K2Node_CallFunction.h"
+#include "K2Node_StructToJsonString.generated.h"
+
+UCLASS()
+class UK2Node_StructToJsonString : public UK2Node_CallFunction
+{
+	GENERATED_BODY()
+	
+public:
+
+	UK2Node_StructToJsonString();
+
+	// //~ Begin UEdGraphNode Interface.
+	virtual void AllocateDefaultPins() override;
+	virtual FText GetTooltipText() const override;	
+	// //~ End UEdGraphNode Interface.
+
+	//~ Begin K2Node Interface
+	virtual bool IsNodePure() const override;
+	virtual void GetMenuActions(FBlueprintActionDatabaseRegistrar& InActionRegistrar) const override;
+	virtual FText GetMenuCategory() const override;
+	//~ End K2Node Interface
+
+	protected:
+	void NotifyInputChanged() const;
+};
diff --git a/Engine/Plugins/JsonBlueprintUtilities/Source/JsonBlueprintUtilities/Private/JsonBlueprintFunctionLibrary.cpp b/Engine/Plugins/JsonBlueprintUtilities/Source/JsonBlueprintUtilities/Private/JsonBlueprintFunctionLibrary.cpp
--- a/Engine/Plugins/JsonBlueprintUtilities/Source/JsonBlueprintUtilities/Private/JsonBlueprintFunctionLibrary.cpp
+++ b/Engine/Plugins/JsonBlueprintUtilities/Source/JsonBlueprintUtilities/Private/JsonBlueprintFunctionLibrary.cpp
@@ -178,6 +178,41 @@ DEFINE_FUNCTION(UJsonBlueprintFunctionLibrary::execSetField)
 	*StaticCast<bool*>(RESULT_PARAM) = bResult;
 }
 
+DEFINE_FUNCTION(UJsonBlueprintFunctionLibrary::execStructToJsonString)
+{
+	Stack.StepCompiledIn<FProperty>(nullptr);
+	FProperty* ValueProperty = Stack.MostRecentProperty;
+	void* ValuePtr = Stack.MostRecentPropertyAddress;
+
+	PARAM_PASSED_BY_REF(OutJsonString, FStrProperty, FString);
+
+	P_FINISH;
+
+	if (!ValueProperty || !ValuePtr)
+	{
+		const FBlueprintExceptionInfo ExceptionInfo(
+			EBlueprintExceptionType::AccessViolation,
+			LOCTEXT("StructToJsonString_MissingInputProperty", "Failed to resolve the input parameter for StructToJsonString.")
+		);
+		FBlueprintCoreDelegates::ThrowScriptException(P_THIS, Stack, ExceptionInfo);
+	}
+
+	bool bResult;
+	FStructProperty* const StructProperty = CastField<FStructProperty>(ValueProperty);
+	if (!StructProperty)
+	{
+		bResult = false;
+		*StaticCast<bool*>(RESULT_PARAM) = bResult;
+		return;
+	}
+	
+	P_NATIVE_BEGIN
+	bResult = FJsonObjectConverter::UStructToJsonObjectString(StructProperty->Struct, ValuePtr, OutJsonString);
+	P_NATIVE_END
+
+	*StaticCast<bool*>(RESULT_PARAM) = bResult;
+}
+
 bool UJsonBlueprintFunctionLibrary::HasField(const FJsonObjectWrapper& JsonObject, const FString& FieldName)
 {
 	return JsonObject.JsonObject->HasField(FieldName);
@@ -218,10 +253,14 @@ bool UJsonBlueprintFunctionLibrary::PropertyToJsonField(
 	const FString& FieldName,
 	FProperty* SourceProperty,
 	const void* SourceValuePtr,
-	const FJsonObjectWrapper& TargetObject)
+	FJsonObjectWrapper& TargetObject)
 {
 	check(SourceProperty && SourceValuePtr);
-	check(TargetObject.JsonObject.IsValid());
+
+	if(!TargetObject.JsonObject.IsValid())
+	{
+		TargetObject.JsonObject = MakeShared<FJsonObject>();
+	}
 
 	TargetObject.JsonObject->SetField(FieldName, FJsonObjectConverter::UPropertyToJsonValue(SourceProperty, SourceValuePtr));
 	return true;
diff --git a/Engine/Plugins/JsonBlueprintUtilities/Source/JsonBlueprintUtilities/Public/JsonBlueprintFunctionLibrary.h b/Engine/Plugins/JsonBlueprintUtilities/Source/JsonBlueprintUtilities/Public/JsonBlueprintFunctionLibrary.h
--- a/Engine/Plugins/JsonBlueprintUtilities/Source/JsonBlueprintUtilities/Public/JsonBlueprintFunctionLibrary.h
+++ b/Engine/Plugins/JsonBlueprintUtilities/Source/JsonBlueprintUtilities/Public/JsonBlueprintFunctionLibrary.h
@@ -41,6 +41,10 @@ public:
 	static UPARAM(DisplayName="Success") bool SetField(const FJsonObjectWrapper& JsonObject, const FString& FieldName, const int32& Value);
 	DECLARE_FUNCTION(execSetField);
 
+	UFUNCTION(BlueprintCallable, CustomThunk, BlueprintInternalUseOnly, Category="Json", meta = (DisplayName = "Convert Struct To Json String", CustomStructureParam = "Struct", AutoCreateRefTerm = "Struct"))
+	static UPARAM(DisplayName="Success") bool StructToJsonString(const int32& Struct, FString& OutJsonString);
+	DECLARE_FUNCTION(execStructToJsonString);
+	
 	/** Checks if the field exists on the object. */
 	UFUNCTION(BlueprintCallable, Category="Json")
 	static UPARAM(DisplayName="Success") bool HasField(const FJsonObjectWrapper& JsonObject, const FString& FieldName);
@@ -54,5 +58,5 @@ private:
 	static bool JsonFieldToProperty(const FString& FieldName, const FJsonObjectWrapper& SourceObject, FProperty* TargetProperty, void* TargetValuePtr);
 
 	// If FieldName empty, Object created as root, or created with name "Value" otherwise. 
-	static bool PropertyToJsonField(const FString& FieldName, FProperty* SourceProperty, const void* SourceValuePtr, const FJsonObjectWrapper& TargetObject);
+	static bool PropertyToJsonField(const FString& FieldName, FProperty* SourceProperty, const void* SourceValuePtr, FJsonObjectWrapper& TargetObject);
 };
diff --git a/Engine/Plugins/Web/HttpBlueprint/HttpBlueprint.uplugin b/Engine/Plugins/Web/HttpBlueprint/HttpBlueprint.uplugin
new file mode 100644
--- /dev/null
+++ b/Engine/Plugins/Web/HttpBlueprint/HttpBlueprint.uplugin
@@ -0,0 +1,35 @@
+{
+	"FileVersion": 3,
+	"Version": 1,
+	"VersionName": "1.0",
+	"FriendlyName": "HttpBlueprint",
+	"Description": "Allows for sending and receiving HTTP requests in Blueprint",
+	"Category": "Web",
+	"CreatedBy": "Epic Games",
+	"CreatedByURL": "",
+	"DocsURL": "",
+	"MarketplaceURL": "",
+	"SupportURL": "",
+	"CanContainContent": true,
+	"IsBetaVersion": true,
+	"IsExperimentalVersion": false,
+	"Installed": false,
+	"Modules": [
+		{
+			"Name": "HttpBlueprint",
+			"Type": "Runtime",
+			"LoadingPhase": "Default"
+		},
+		{
+			"Name": "HttpBlueprintGraph",
+			"Type": "UncookedOnly",
+			"LoadingPhase": "Default"
+		}
+	],
+	"Plugins": [
+		{
+			"Name": "JsonBlueprintUtilities",
+			"Enabled": true
+		}
+	]
+}
diff --git a/Engine/Plugins/Web/HttpBlueprint/Source/HttpBlueprint/HttpBlueprint.Build.cs b/Engine/Plugins/Web/HttpBlueprint/Source/HttpBlueprint/HttpBlueprint.Build.cs
new file mode 100644
--- /dev/null
+++ b/Engine/Plugins/Web/HttpBlueprint/Source/HttpBlueprint/HttpBlueprint.Build.cs
@@ -0,0 +1,29 @@
+// Copyright Epic Games, Inc. All Rights Reserved.
+
+using UnrealBuildTool;
+
+public class HttpBlueprint : ModuleRules
+{
+	public HttpBlueprint(ReadOnlyTargetRules Target) : base(Target)
+	{
+		PCHUsage = ModuleRules.PCHUsageMode.UseExplicitOrSharedPCHs;
+
+		PublicDependencyModuleNames.AddRange(
+			new string[]
+			{
+				"Core",
+				"CoreUObject",
+				"Engine",
+			}
+		);
+		
+		PrivateDependencyModuleNames.AddRange(
+			new string[]
+			{
+				"HTTP",
+				"Slate",
+				"SlateCore",
+			}
+		);
+	}
+}
diff --git a/Engine/Plugins/Web/HttpBlueprint/Source/HttpBlueprint/Private/HttpBlueprintFunctionLibrary.cpp b/Engine/Plugins/Web/HttpBlueprint/Source/HttpBlueprint/Private/HttpBlueprintFunctionLibrary.cpp
new file mode 100644
--- /dev/null
+++ b/Engine/Plugins/Web/HttpBlueprint/Source/HttpBlueprint/Private/HttpBlueprintFunctionLibrary.cpp
@@ -0,0 +1,65 @@
+// Copyright Epic Games, Inc. All Rights Reserved.
+
+#include "HttpBlueprintFunctionLibrary.h"
+
+void UHttpBlueprintFunctionLibrary::MakeRequestHeader(const TMap<FString, FString>& Headers, FHttpHeader& OutHeader)
+{
+	OutHeader = FHttpHeader{}.SetHeaders(Headers);
+}
+
+bool UHttpBlueprintFunctionLibrary::GetHeaderValue(const FHttpHeader& HeaderObject, FString HeaderName, FString& OutHeaderValue)
+{
+	if (!HeaderObject.IsValid())
+	{
+		FFrame::KismetExecutionMessage(*FString::Printf(TEXT("HeaderObject can not be invalid")), ELogVerbosity::Error);
+		return false;
+	}
+
+	if (HeaderName.IsEmpty())
+	{
+		FFrame::KismetExecutionMessage(*FString::Printf(TEXT("Header Name can not be empty")), ELogVerbosity::Error);
+		return false;
+	}
+
+	OutHeaderValue = HeaderObject.GetHeader(MoveTemp(HeaderName));
+	return !OutHeaderValue.IsEmpty();
+}
+
+TArray<FString> UHttpBlueprintFunctionLibrary::GetAllHeaders(const FHttpHeader& HeaderObject)
+{
+	return HeaderObject.GetAllHeaders();
+}
+
+TMap<FString, FString> UHttpBlueprintFunctionLibrary::GetAllHeaders_Map(const FHttpHeader& HeaderObject)
+{
+	return HeaderObject.GetAllHeadersAsMap();
+}
+
+void UHttpBlueprintFunctionLibrary::AddHeader(FHttpHeader& HeaderObject, FString NewHeader, FString NewHeaderValue)
+{
+	if (NewHeader.IsEmpty())
+	{
+		FFrame::KismetExecutionMessage(*FString::Printf(TEXT("NewHeader can not be empty")), ELogVerbosity::Error);
+		return;
+	}
+
+	if (NewHeaderValue.IsEmpty())
+	{
+		FFrame::KismetExecutionMessage(*FString::Printf(TEXT("NewHeaderValue can not be empty")), ELogVerbosity::Error);
+		return;
+	}
+	
+	TPair<FString, FString> NewPair(NewHeader, NewHeaderValue);
+	HeaderObject.AddHeader(MoveTemp(NewPair));
+}
+
+bool UHttpBlueprintFunctionLibrary::RemoveHeader(FHttpHeader& HeaderObject, FString HeaderToRemove)
+{
+	if (HeaderToRemove.IsEmpty())
+	{
+		FFrame::KismetExecutionMessage(*FString::Printf(TEXT("HeaderToRemove can not be empty")), ELogVerbosity::Error);
+		return false;
+	}
+	
+	return HeaderObject.RemoveHeader(MoveTemp(HeaderToRemove));
+}
\ No newline at end of file
diff --git a/Engine/Plugins/Web/HttpBlueprint/Source/HttpBlueprint/Private/HttpBlueprintModule.cpp b/Engine/Plugins/Web/HttpBlueprint/Source/HttpBlueprint/Private/HttpBlueprintModule.cpp
new file mode 100644
--- /dev/null
+++ b/Engine/Plugins/Web/HttpBlueprint/Source/HttpBlueprint/Private/HttpBlueprintModule.cpp
@@ -0,0 +1,6 @@
+// Copyright Epic Games, Inc. All Rights Reserved.
+
+#include "CoreMinimal.h"
+#include "Modules/ModuleManager.h"
+
+IMPLEMENT_MODULE(FDefaultModuleImpl, HttpBlueprint)
diff --git a/Engine/Plugins/Web/HttpBlueprint/Source/HttpBlueprint/Private/HttpHeader.cpp b/Engine/Plugins/Web/HttpBlueprint/Source/HttpBlueprint/Private/HttpHeader.cpp
new file mode 100644
--- /dev/null
+++ b/Engine/Plugins/Web/HttpBlueprint/Source/HttpBlueprint/Private/HttpHeader.cpp
@@ -0,0 +1,58 @@
+﻿// Copyright Epic Games, Inc. All Rights Reserved.
+
+#include "HttpHeader.h"
+
+#include "Interfaces/IHttpRequest.h"
+
+FHttpHeader FHttpHeader::SetHeaders(const TMap<FString, FString>& NewHeaders)
+{
+	Headers = NewHeaders;
+	return *this;
+}
+
+FString FHttpHeader::GetHeader(FString&& Key) const
+{
+	const FString* FoundHeader = Headers.Find(MoveTemp(Key));
+	return FoundHeader ? *FoundHeader : TEXT("");
+}
+
+void FHttpHeader::AddHeader(TPair<FString, FString>&& NewHeader)
+{
+	Headers.Add(MoveTemp(NewHeader));
+}
+
+bool FHttpHeader::RemoveHeader(FString&& HeaderToRemove)
+{
+	return !!Headers.Remove(HeaderToRemove);
+}
+
+TArray<FString> FHttpHeader::GetAllHeaders() const
+{
+	TArray<FString> OutArray;
+	OutArray.Reserve(Headers.Num());
+
+	for (const TPair<FString, FString>& Header : Headers)
+	{
+		OutArray.Add(FString::Printf(TEXT("%s: %s"), *Header.Key, *Header.Value));
+	}
+
+	return OutArray;
+}
+
+const TMap<FString, FString>& FHttpHeader::GetAllHeadersAsMap() const
+{
+	return Headers;
+}
+
+bool FHttpHeader::IsValid() const
+{
+	return Headers.Num() > 0;
+}
+
+void FHttpHeader::AssignHeadersToRequest(const TSharedRef<IHttpRequest> Request)
+{
+	for (const TPair<FString, FString>& HeaderPair : Headers)
+	{
+		Request->SetHeader(HeaderPair.Key, HeaderPair.Value);
+	}
+}
diff --git a/Engine/Plugins/Web/HttpBlueprint/Source/HttpBlueprint/Public/HttpBlueprintFunctionLibrary.h b/Engine/Plugins/Web/HttpBlueprint/Source/HttpBlueprint/Public/HttpBlueprintFunctionLibrary.h
new file mode 100644
--- /dev/null
+++ b/Engine/Plugins/Web/HttpBlueprint/Source/HttpBlueprint/Public/HttpBlueprintFunctionLibrary.h
@@ -0,0 +1,58 @@
+// Copyright Epic Games, Inc. All Rights Reserved.
+
+#pragma once
+
+#include "CoreMinimal.h"
+#include "HttpHeader.h"
+#include "Kismet/BlueprintFunctionLibrary.h"
+
+#include "HttpBlueprintFunctionLibrary.generated.h"
+
+UCLASS()
+class HTTPBLUEPRINT_API UHttpBlueprintFunctionLibrary : public UBlueprintFunctionLibrary
+{
+	GENERATED_BODY()
+
+public:
+	UFUNCTION(BlueprintPure, BlueprintInternalUseOnly, Category = "Http")
+	static void MakeRequestHeader(const TMap<FString, FString>& Headers, UPARAM(DisplayName="Header") FHttpHeader& OutHeader);
+
+	/**
+	 *	Get the value associated with a Header name
+	 *
+	 *	@Param HeaderObject: The structure that contains all of the headers
+	 *	@Param HeaderName: The name of the Header
+	 *	@Param OutHeaderValue: The value of the Header. Empty if the Header could not be found
+	 *	@Return Whether the operation was successful
+	 */
+	UFUNCTION(BlueprintPure, Category = "Http")
+	static UPARAM(DisplayName="Success") bool GetHeaderValue(const FHttpHeader& HeaderObject, FString HeaderName, FString& OutHeaderValue);
+
+	/** Returns all of the headers and their values by value */
+	UFUNCTION(BlueprintPure, Category = "Http")
+	static UPARAM(DisplayName="Headers") TArray<FString> GetAllHeaders(const FHttpHeader& HeaderObject);
+
+	/** Returns all of the headers and their values as a map by value */
+	UFUNCTION(BlueprintPure, Category = "Http", meta = (DisplayName = "GetAllHeadersAsMap"))
+	static UPARAM(DisplayName="Headers") TMap<FString, FString> GetAllHeaders_Map(const FHttpHeader& HeaderObject);
+	
+	/**
+	 *	Adds a new Header
+	 *
+	 *	@Param HeaderObject: The structure that contains all of the headers
+	 *	@Param NewHeader: The name of the header
+	 *	@Param NewHeaderValue: The value of the new header
+	 */
+	UFUNCTION(BlueprintCallable, Category = "Http")
+	static void AddHeader(UPARAM(Ref) FHttpHeader& HeaderObject, FString NewHeader, FString NewHeaderValue);
+
+	/**
+	 *	Removes a header from the HeaderObject
+	 *
+	 *	@Param HeaderObject: The structure that contains all of the headers
+	 *	@Param HeaderToRemove: The Key of the header to remove
+	 *	@Return Whether or not the operation was successful. A value of false most likely means the Key did not exist on the HeaderObject
+	 */
+	UFUNCTION(BlueprintCallable, Category = "Http")
+	static bool RemoveHeader(UPARAM(Ref) FHttpHeader& HeaderObject, FString HeaderToRemove);
+};
diff --git a/Engine/Plugins/Web/HttpBlueprint/Source/HttpBlueprint/Public/HttpBlueprintTypes.h b/Engine/Plugins/Web/HttpBlueprint/Source/HttpBlueprint/Public/HttpBlueprintTypes.h
new file mode 100644
--- /dev/null
+++ b/Engine/Plugins/Web/HttpBlueprint/Source/HttpBlueprint/Public/HttpBlueprintTypes.h
@@ -0,0 +1,32 @@
+﻿// Copyright Epic Games, Inc. All Rights Reserved.
+
+#pragma once
+
+#include "CoreMinimal.h"
+
+#include "HttpBlueprintTypes.generated.h"
+
+UENUM(BlueprintType)
+enum class EHttpVerbs : uint8
+{
+	Post	= 0,
+	Put		= 1,
+	Delete	= 2,
+	Patch	= 3,
+	// @note: Anything past Patch will not display the input body pin
+	Get		= 4,
+
+	MAX		= 255	UMETA(Hidden),
+};
+
+ENUM_CLASS_FLAGS(EHttpVerbs);
+
+UENUM(BlueprintType)
+enum class ERequestPresets : uint8
+{
+	// The order here matters.	
+	Json			UMETA(DisplayName = "Json Request"),
+	Http			UMETA(DisplayName = "Http Request"),
+	Url				UMETA(DisplayName = "Url Encoded Request"),
+	Custom			UMETA(DisplayName = "Custom Request")
+};
diff --git a/Engine/Plugins/Web/HttpBlueprint/Source/HttpBlueprint/Public/HttpHeader.h b/Engine/Plugins/Web/HttpBlueprint/Source/HttpBlueprint/Public/HttpHeader.h
new file mode 100644
--- /dev/null
+++ b/Engine/Plugins/Web/HttpBlueprint/Source/HttpBlueprint/Public/HttpHeader.h
@@ -0,0 +1,63 @@
+﻿// Copyright Epic Games, Inc. All Rights Reserved.
+
+#pragma once
+
+#include "CoreMinimal.h"
+
+#include "HttpHeader.generated.h"
+
+/**
+ * Provides a way for blueprint to create and store a map of HTTP headers
+ */
+USTRUCT(BlueprintType)
+struct HTTPBLUEPRINT_API FHttpHeader
+{
+	GENERATED_BODY()
+
+	/**
+	 * Sets the internal headers.
+	 * Will overwrite any headers already in the map
+	 */
+	UE_NODISCARD FHttpHeader SetHeaders(const TMap<FString, FString>& NewHeaders);
+
+	/**
+	 * Returns the header value associated with the Key
+	 */
+	UE_NODISCARD FString GetHeader(FString&& Key) const;
+
+	/**
+	 * Adds a new header key/value pair to the internal header map
+	 */
+	void AddHeader(TPair<FString, FString>&& NewHeader);
+
+	/**
+	 * Removes a header from the internal header map
+	 */
+	UE_NODISCARD bool RemoveHeader(FString&& HeaderToRemove);
+
+	/**
+	 * Returns the internal header map as an array
+	 */
+	UE_NODISCARD TArray<FString> GetAllHeaders() const;
+
+	/**
+	 * Returns a copy of the internal header map
+	 */
+	UE_NODISCARD inline const TMap<FString, FString>& GetAllHeadersAsMap() const;
+
+	/**
+	 * Returns if the header is valid or not. This is done by checking the map size
+	 */
+	UE_NODISCARD inline bool IsValid() const;
+
+	/**
+	 * Convenience function to set all of the headers in the internal header map on the request object.
+	 */
+	void AssignHeadersToRequest(TSharedRef<class IHttpRequest> Request);
+	
+private:
+
+	UPROPERTY()
+	TMap<FString, FString> Headers;
+	
+};
diff --git a/Engine/Plugins/Web/HttpBlueprint/Source/HttpBlueprintGraph/HttpBlueprintGraph.Build.cs b/Engine/Plugins/Web/HttpBlueprint/Source/HttpBlueprintGraph/HttpBlueprintGraph.Build.cs
new file mode 100644
--- /dev/null
+++ b/Engine/Plugins/Web/HttpBlueprint/Source/HttpBlueprintGraph/HttpBlueprintGraph.Build.cs
@@ -0,0 +1,48 @@
+// Copyright Epic Games, Inc. All Rights Reserved.
+
+using UnrealBuildTool;
+
+public class HttpBlueprintGraph : ModuleRules
+{
+	public HttpBlueprintGraph(ReadOnlyTargetRules Target) : base(Target)
+	{
+		PCHUsage = ModuleRules.PCHUsageMode.UseExplicitOrSharedPCHs;
+
+		PrivateDependencyModuleNames.AddRange(
+			new string[]
+			{
+				"ApplicationCore",
+				"AssetRegistry",
+				"BlueprintGraph",
+				"Core",
+				"CoreUObject",
+				"Engine",
+				"GraphEditor",
+				"InputCore",
+				"Kismet",
+				"KismetCompiler",
+				"KismetWidgets",
+				"Slate",
+				"SlateCore", 
+				"ToolMenus"
+			}
+		);
+
+		if (Target.bBuildEditor)
+		{
+			PrivateDependencyModuleNames.AddRange(
+				new string[]
+				{
+					"AssetTools",
+					"EditorStyle",
+					"HttpBlueprint",
+					"HTTP",
+					"JsonBlueprintUtilities",
+					"Projects",
+					"PropertyEditor",
+					"UnrealEd",
+				}
+			);
+		}
+	}
+}
\ No newline at end of file
diff --git a/Engine/Plugins/Web/HttpBlueprint/Source/HttpBlueprintGraph/Private/HttpBlueprintGraphLog.cpp b/Engine/Plugins/Web/HttpBlueprint/Source/HttpBlueprintGraph/Private/HttpBlueprintGraphLog.cpp
new file mode 100644
--- /dev/null
+++ b/Engine/Plugins/Web/HttpBlueprint/Source/HttpBlueprintGraph/Private/HttpBlueprintGraphLog.cpp
@@ -0,0 +1,5 @@
+// Copyright Epic Games, Inc. All Rights Reserved.
+
+#include "HttpBlueprintGraphLog.h"
+
+DEFINE_LOG_CATEGORY(LogHttpBlueprintEditor);
\ No newline at end of file
diff --git a/Engine/Plugins/Web/HttpBlueprint/Source/HttpBlueprintGraph/Private/HttpBlueprintGraphLog.h b/Engine/Plugins/Web/HttpBlueprint/Source/HttpBlueprintGraph/Private/HttpBlueprintGraphLog.h
new file mode 100644
--- /dev/null
+++ b/Engine/Plugins/Web/HttpBlueprint/Source/HttpBlueprintGraph/Private/HttpBlueprintGraphLog.h
@@ -0,0 +1,7 @@
+// Copyright Epic Games, Inc. All Rights Reserved.
+
+#pragma once
+
+#include "Logging/LogMacros.h"
+
+DECLARE_LOG_CATEGORY_EXTERN(LogHttpBlueprintEditor, Log, All);
\ No newline at end of file
diff --git a/Engine/Plugins/Web/HttpBlueprint/Source/HttpBlueprintGraph/Private/HttpBlueprintGraphModule.cpp b/Engine/Plugins/Web/HttpBlueprint/Source/HttpBlueprintGraph/Private/HttpBlueprintGraphModule.cpp
new file mode 100644
--- /dev/null
+++ b/Engine/Plugins/Web/HttpBlueprint/Source/HttpBlueprintGraph/Private/HttpBlueprintGraphModule.cpp
@@ -0,0 +1,6 @@
+// Copyright Epic Games, Inc. All Rights Reserved.
+
+#include "CoreMinimal.h"
+#include "Modules/ModuleManager.h"
+
+IMPLEMENT_MODULE(FDefaultModuleImpl, HttpBlueprintGraph)
diff --git a/Engine/Plugins/Web/HttpBlueprint/Source/HttpBlueprintGraph/Private/HttpRequestHeaderPresets.h b/Engine/Plugins/Web/HttpBlueprint/Source/HttpBlueprintGraph/Private/HttpRequestHeaderPresets.h
new file mode 100644
--- /dev/null
+++ b/Engine/Plugins/Web/HttpBlueprint/Source/HttpBlueprintGraph/Private/HttpRequestHeaderPresets.h
@@ -0,0 +1,66 @@
+﻿// Copyright Epic Games, Inc. All Rights Reserved.
+
+#pragma once
+
+#include "CoreMinimal.h"
+
+namespace UE::HttpBlueprint::HeaderPresets
+{
+	struct FBasePreset
+	{
+		FBasePreset()
+		{
+			HeaderPresets =
+			{
+				{TEXT("Accept-Encoding"), TEXT("identity")},
+				{TEXT("User-Agent"), TEXT("X-UnrealEngine-Agent")}
+			};
+		}
+
+		void SetDefaults(const TMap<FString, FString>& AdditionalHeaders)
+		{
+			HeaderPresets =
+			{
+				{TEXT("Accept-Encoding"), TEXT("identity")},
+				{TEXT("User-Agent"), TEXT("X-UnrealEngine-Agent")}
+			};
+
+			HeaderPresets.Append(AdditionalHeaders);
+		}
+
+		TMap<FString, FString> HeaderPresets;
+	};
+
+	struct FJsonPreset : FBasePreset
+	{
+		FJsonPreset()
+		{
+			SetDefaults({
+				{TEXT("Accepts"), TEXT("application/json")},
+				{TEXT("Content-Type"), TEXT("application/json")}
+			});
+		}
+	};
+
+	struct FHtmlPreset : FBasePreset
+	{
+		FHtmlPreset()
+		{
+			SetDefaults({
+				{TEXT("Accepts"), TEXT("text/html")},
+				{TEXT("Content-Type"), TEXT("text/html")}
+			});
+		}
+	};
+
+	struct FUrlPreset : FBasePreset
+	{
+		FUrlPreset()
+		{
+			SetDefaults({
+				{TEXT("Accepts"), TEXT("application/x-www-form-urlencoded")},
+				{TEXT("Content-Type"), TEXT("application/x-www-form-urlencoded")}
+			});
+		}
+	};
+}
diff --git a/Engine/Plugins/Web/HttpBlueprint/Source/HttpBlueprintGraph/Private/HttpRequestProxyObject.cpp b/Engine/Plugins/Web/HttpBlueprint/Source/HttpBlueprintGraph/Private/HttpRequestProxyObject.cpp
new file mode 100644
--- /dev/null
+++ b/Engine/Plugins/Web/HttpBlueprint/Source/HttpBlueprintGraph/Private/HttpRequestProxyObject.cpp
@@ -0,0 +1,42 @@
+// Copyright Epic Games, Inc. All Rights Reserved.
+
+#include "HttpRequestProxyObject.h"
+
+#include "HttpBlueprintFunctionLibrary.h"
+#include "HttpModule.h"
+#include "Interfaces/IHttpResponse.h"
+
+UHttpRequestProxyObject* UHttpRequestProxyObject::CreateHttpRequestProxyObject(
+	const FString& InUrl,
+	const FString& InVerb,
+	FHttpHeader InHeader,
+	const FString& InBody)
+{
+	UHttpRequestProxyObject* const Proxy = NewObject<UHttpRequestProxyObject>();
+	Proxy->SetFlags(RF_StrongRefOnFrame);
+	Proxy->CachedHeader = InHeader;
+	Proxy->ProcessRequest(InUrl, InVerb, MoveTemp(InHeader), InBody);
+	return Proxy;
+}
+
+void UHttpRequestProxyObject::ProcessRequest(
+	const FString& InUrl,
+	const FString& InVerb,
+	FHttpHeader&& InHeader,
+	const FString& InBody)
+{
+	FHttpModule& HttpModule = FHttpModule::Get();
+	const TSharedRef<IHttpRequest> Request = HttpModule.CreateRequest();
+	Request->SetURL(InUrl);
+	Request->SetVerb(InVerb);
+	Request->SetContentAsString(InBody);
+	InHeader.AssignHeadersToRequest(Request);
+	Request->ProcessRequest();
+
+	Request->OnProcessRequestComplete().BindUObject(this, &ThisClass::ProcessComplete);
+}
+
+void UHttpRequestProxyObject::ProcessComplete(FHttpRequestPtr InRequest, FHttpResponsePtr InResponse, bool bInSuccessful)
+{
+	OnRequestComplete.Broadcast(InResponse->GetContentAsString(), bInSuccessful, MoveTemp(CachedHeader));
+}
diff --git a/Engine/Plugins/Web/HttpBlueprint/Source/HttpBlueprintGraph/Private/HttpRequestProxyObject.h b/Engine/Plugins/Web/HttpBlueprint/Source/HttpBlueprintGraph/Private/HttpRequestProxyObject.h
new file mode 100644
--- /dev/null
+++ b/Engine/Plugins/Web/HttpBlueprint/Source/HttpBlueprintGraph/Private/HttpRequestProxyObject.h
@@ -0,0 +1,36 @@
+// Copyright Epic Games, Inc. All Rights Reserved.
+
+#pragma once
+
+#include "CoreMinimal.h"
+#include "HttpHeader.h"
+#include "Interfaces/IHttpRequest.h"
+
+#include "HttpRequestProxyObject.generated.h"
+
+DECLARE_DYNAMIC_MULTICAST_DELEGATE_ThreeParams(FOnRequestComplete, FString, Response, bool, bSuccessful, FHttpHeader, OutHeader);
+
+UCLASS(MinimalAPI)
+class UHttpRequestProxyObject : public UObject
+{
+	GENERATED_BODY()
+
+public:
+	UPROPERTY(BlueprintAssignable)
+	FOnRequestComplete OnRequestComplete;
+
+	UFUNCTION(BlueprintCallable, Meta = (BlueprintInternalUseOnly, DisplayName = "Http Request"), Category = "Http")
+	static UHttpRequestProxyObject* CreateHttpRequestProxyObject(
+		const FString& InUrl,
+		const FString& InVerb,
+		FHttpHeader InHeader,
+		const FString& InBody);
+
+protected:
+	void ProcessRequest(const FString& InUrl, const FString& InVerb, FHttpHeader&& InHeader, const FString& InBody);
+	void ProcessComplete(FHttpRequestPtr InRequest, FHttpResponsePtr InResponse, bool bInSuccessful);
+
+private:
+	UPROPERTY()
+	FHttpHeader CachedHeader;
+};
diff --git a/Engine/Plugins/Web/HttpBlueprint/Source/HttpBlueprintGraph/Private/K2Node_AsyncMakeRequestHeader.cpp b/Engine/Plugins/Web/HttpBlueprint/Source/HttpBlueprintGraph/Private/K2Node_AsyncMakeRequestHeader.cpp
new file mode 100644
--- /dev/null
+++ b/Engine/Plugins/Web/HttpBlueprint/Source/HttpBlueprintGraph/Private/K2Node_AsyncMakeRequestHeader.cpp
@@ -0,0 +1,45 @@
+// Copyright Epic Games, Inc. All Rights Reserved.
+
+#include "K2Node_AsyncMakeRequestHeader.h"
+
+#include "HttpRequestProxyObject.h"
+
+UK2Node_AsyncMakeRequestHeader::UK2Node_AsyncMakeRequestHeader()
+{
+	ProxyFactoryFunctionName = GET_FUNCTION_NAME_CHECKED(UHttpRequestProxyObject, CreateHttpRequestProxyObject);
+	ProxyFactoryClass = UHttpRequestProxyObject::StaticClass();
+	ProxyClass = UHttpRequestProxyObject::StaticClass();
+}
+
+UEdGraphPin* UK2Node_AsyncMakeRequestHeader::GetVerbPin() const
+{
+	UEdGraphPin* Pin = FindPinChecked(TEXT("InVerb"));
+	check(Pin->Direction == EGPD_Input);
+	return Pin;
+}
+
+UEdGraphPin* UK2Node_AsyncMakeRequestHeader::GetBodyPin() const
+{
+	UEdGraphPin* Pin = FindPinChecked(TEXT("InBody"));
+	check(Pin->Direction == EGPD_Input);
+	return Pin;
+}
+
+UEdGraphPin* UK2Node_AsyncMakeRequestHeader::GetHeaderPin() const
+{
+	UEdGraphPin* Pin = FindPinChecked(TEXT("InHeader"));
+	check(Pin->Direction == EGPD_Input);
+	return Pin;
+}
+
+UEdGraphPin* UK2Node_AsyncMakeRequestHeader::GetUrlPin() const
+{
+	UEdGraphPin* Pin = FindPinChecked(TEXT("InUrl"));
+	check(Pin->Direction == EGPD_Input);
+	return Pin;
+}
+
+void UK2Node_AsyncMakeRequestHeader::GetMenuActions(FBlueprintActionDatabaseRegistrar& ActionRegistrar) const
+{
+	// We only want our custom K2Node version to show up, this version should be hidden from the user
+}
diff --git a/Engine/Plugins/Web/HttpBlueprint/Source/HttpBlueprintGraph/Private/K2Node_AsyncMakeRequestHeader.h b/Engine/Plugins/Web/HttpBlueprint/Source/HttpBlueprintGraph/Private/K2Node_AsyncMakeRequestHeader.h
new file mode 100644
--- /dev/null
+++ b/Engine/Plugins/Web/HttpBlueprint/Source/HttpBlueprintGraph/Private/K2Node_AsyncMakeRequestHeader.h
@@ -0,0 +1,24 @@
+// Copyright Epic Games, Inc. All Rights Reserved.
+
+#pragma once
+
+#include "K2Node_BaseAsyncTask.h"
+
+#include "K2Node_AsyncMakeRequestHeader.generated.h"
+
+UCLASS()
+class UK2Node_AsyncMakeRequestHeader final : public UK2Node_BaseAsyncTask
+{
+	GENERATED_BODY()
+
+public:
+	UK2Node_AsyncMakeRequestHeader();
+
+	UEdGraphPin* GetVerbPin() const;
+    UEdGraphPin* GetBodyPin() const;
+    UEdGraphPin* GetHeaderPin() const;
+    UEdGraphPin* GetUrlPin() const;
+
+protected:
+	virtual void GetMenuActions(FBlueprintActionDatabaseRegistrar& ActionRegistrar) const override;
+};
diff --git a/Engine/Plugins/Web/HttpBlueprint/Source/HttpBlueprintGraph/Private/K2Node_HttpRequest.cpp b/Engine/Plugins/Web/HttpBlueprint/Source/HttpBlueprintGraph/Private/K2Node_HttpRequest.cpp
new file mode 100644
--- /dev/null
+++ b/Engine/Plugins/Web/HttpBlueprint/Source/HttpBlueprintGraph/Private/K2Node_HttpRequest.cpp
@@ -0,0 +1,483 @@
+// Copyright Epic Games, Inc. All Rights Reserved.
+
+#include "K2Node_HttpRequest.h"
+
+#include "BlueprintActionDatabaseRegistrar.h"
+#include "BlueprintNodeSpawner.h"
+#include "HttpBlueprintFunctionLibrary.h"
+#include "HttpBlueprintTypes.h"
+#include "HttpRequestProxyObject.h"
+#include "JsonBlueprintFunctionLibrary.h"
+#include "K2Node_AsyncMakeRequestHeader.h"
+#include "K2Node_CallFunction.h"
+#include "K2Node_IfThenElse.h"
+#include "KismetCompiler.h"
+#include "Kismet2/BlueprintEditorUtils.h"
+
+#define LOCTEXT_NAMESPACE "K2Node_HttpRequest"
+
+namespace HttpRequestLiterals
+{
+	static const FName SuccessPinName(TEXT("On Success"));
+	static const FName ErrorPinName(TEXT("On Error"));
+	static const FName UrlPinName(TEXT("Url"));
+	static const FName VerbPinName(TEXT("Verb"));
+	static const FName HeaderPinName(TEXT("Header"));
+	static const FName OutHeaderPinName(TEXT("OutHeader"));
+	static const FName BodyPinName(TEXT("Body"));
+	static const FName OutputBodyPinName(TEXT("Result Body"));
+}
+
+UK2Node_HttpRequest::UK2Node_HttpRequest(const FObjectInitializer& ObjectInitializer)
+	: Super(ObjectInitializer)
+	, VerbEnum(FindObject<UEnum>(ANY_PACKAGE, TEXT("EHttpVerbs")))
+{
+}
+
+void UK2Node_HttpRequest::AllocateDefaultPins()
+{
+	const UEdGraphSchema_K2* Schema = GetDefault<UEdGraphSchema_K2>();
+	
+	Super::AllocateDefaultPins();
+	
+	// Execution Pins
+	CreatePin(EGPD_Input, UEdGraphSchema_K2::PC_Exec, UEdGraphSchema_K2::PN_Execute);
+	
+	UEdGraphPin* const ThenPin = CreatePin(EGPD_Output, UEdGraphSchema_K2::PC_Exec, UEdGraphSchema_K2::PN_Then);
+	ThenPin->PinFriendlyName = LOCTEXT("HttpRequest ThenPin", "Request Processing");
+
+	CreatePin(EGPD_Output,
+			UEdGraphSchema_K2::PC_Exec,
+			HttpRequestLiterals::SuccessPinName)
+		->PinFriendlyName = LOCTEXT("HttpRequest SuccessPin", "Request Was Successful");
+
+	CreatePin(EGPD_Output,
+			UEdGraphSchema_K2::PC_Exec,
+			HttpRequestLiterals::ErrorPinName)
+		->PinFriendlyName = LOCTEXT("HttpRequest FailurePin", "Request Was Not Successful");
+	
+	// Input pins
+	CreatePin(EGPD_Input, UEdGraphSchema_K2::PC_String, HttpRequestLiterals::UrlPinName);
+	UEdGraphPin* EnumPin = CreatePin(EGPD_Input, UEdGraphSchema_K2::PC_Byte, VerbEnum.Get(), HttpRequestLiterals::VerbPinName);
+	Schema->SetPinDefaultValueAtConstruction(EnumPin, VerbEnum->GetNameStringByIndex(0));
+
+	FCreatePinParams PinParams;
+	PinParams.bIsConst = true;
+	CreatePin(EGPD_Input, UEdGraphSchema_K2::PC_Struct, FHttpHeader::StaticStruct(), HttpRequestLiterals::HeaderPinName, PinParams);
+
+	// Create BodyInputPin. Will only create if the EnumPin has the correct value
+	HandleBodyInputPin();
+
+	// Output pins
+	CreatePin(EGPD_Output,
+			UEdGraphSchema_K2::PC_Struct,
+			FHttpHeader::StaticStruct(),
+			HttpRequestLiterals::OutHeaderPinName)
+		->PinFriendlyName = LOCTEXT("HttpRequest OutHeader", "Header");
+
+	UEdGraphPin* const OutBodyPin = CreatePin(EGPD_Output, UEdGraphSchema_K2::PC_Wildcard, HttpRequestLiterals::OutputBodyPinName);
+	SetPinToolTip(OutBodyPin, LOCTEXT("HttpRequest OutBodyPin", "The response from the request"));
+}
+
+FText UK2Node_HttpRequest::GetTooltipText() const
+{
+	return LOCTEXT("HttpRequest_Tooltip", "Processes a Http request ");
+}
+
+FText UK2Node_HttpRequest::GetNodeTitle(ENodeTitleType::Type Title) const
+{
+	const UEdGraphPin* const VerbEnumPin = FindPin(HttpRequestLiterals::VerbPinName);
+	return FText::Format(
+		LOCTEXT("HttpRequest_Title", "Http {0} Request"),
+		FText::FromString(VerbEnumPin ? *VerbEnumPin->DefaultValue : TEXT("")));
+}
+
+
+bool UK2Node_HttpRequest::IsCompatibleWithGraph(const UEdGraph* Graph) const
+{
+	return Graph->GetSchema()->GetGraphType(Graph) == GT_Ubergraph && Super::IsCompatibleWithGraph(Graph);
+}
+
+void UK2Node_HttpRequest::PinConnectionListChanged(UEdGraphPin* Pin)
+{
+	Super::PinConnectionListChanged(Pin);
+	
+	Modify();
+
+	SyncBodyPinType(GetBodyPin());
+	SyncBodyPinType(GetOutBodyPin());
+}
+
+void UK2Node_HttpRequest::PinDefaultValueChanged(UEdGraphPin* Pin)
+{
+	Super::PinDefaultValueChanged(Pin);
+
+	if (Pin->PinName == HttpRequestLiterals::VerbPinName)
+	{
+		HandleBodyInputPin();
+	}
+}
+
+void UK2Node_HttpRequest::PinTypeChanged(UEdGraphPin* Pin)
+{
+	SyncBodyPinType(GetBodyPin());
+	SyncBodyPinType(GetOutBodyPin());
+	
+	Super::PinTypeChanged(Pin);
+}
+
+void UK2Node_HttpRequest::ExpandNode(FKismetCompilerContext& CompilerContext, UEdGraph* SourceGraph)
+{
+	UK2Node_AsyncMakeRequestHeader* AsyncRequestNode = CompilerContext.SpawnIntermediateNode<UK2Node_AsyncMakeRequestHeader>(this, SourceGraph);
+	AsyncRequestNode->AllocateDefaultPins();
+
+	// Move pins from this node to the proxy object node
+	CompilerContext.MovePinLinksToIntermediate(*GetThenPin(), *AsyncRequestNode->FindPinChecked(UEdGraphSchema_K2::PN_Then));
+	CompilerContext.MovePinLinksToIntermediate(*GetVerbPin(), *AsyncRequestNode->GetVerbPin());
+	CompilerContext.MovePinLinksToIntermediate(*GetHeaderPin(), *AsyncRequestNode->GetHeaderPin());
+	CompilerContext.MovePinLinksToIntermediate(*GetUrlPin(), *AsyncRequestNode->GetUrlPin());
+	
+	if (GetBodyPin() && GetBodyPin()->PinType.PinCategory == UEdGraphSchema_K2::PC_String)
+	{
+		CompilerContext.MovePinLinksToIntermediate(*GetExecPin(), *AsyncRequestNode->GetExecPin());
+		CompilerContext.MovePinLinksToIntermediate(*GetBodyPin(), *AsyncRequestNode->GetBodyPin());
+	}
+	else if (GetBodyPin() && GetBodyPin()->PinType.PinCategory == UEdGraphSchema_K2::PC_Struct)
+	{
+		const FName& FromStructFunctionName = GET_FUNCTION_NAME_CHECKED(UJsonBlueprintFunctionLibrary, StructToJsonString);
+		UClass* JsonLibraryClass = UJsonBlueprintFunctionLibrary::StaticClass();
+
+		UK2Node_CallFunction* CallFromStructFunctionNode = SourceGraph->CreateIntermediateNode<UK2Node_CallFunction>();
+		CallFromStructFunctionNode->FunctionReference.SetExternalMember(FromStructFunctionName, JsonLibraryClass);
+		CallFromStructFunctionNode->AllocateDefaultPins();
+		CompilerContext.MessageLog.NotifyIntermediateObjectCreation(CallFromStructFunctionNode, this);
+
+		UEdGraphPin* InStructPin = CallFromStructFunctionNode->FindPinChecked(TEXT("InStruct"));
+		InStructPin->PinType = GetBodyPin()->PinType;
+		
+		UEdGraphPin* OutJsonPin = CallFromStructFunctionNode->FindPinChecked(TEXT("OutJson"));
+
+		CompilerContext.MovePinLinksToIntermediate(*GetExecPin(), *CallFromStructFunctionNode->GetExecPin());
+		CompilerContext.MovePinLinksToIntermediate(*GetBodyPin(), *InStructPin);
+		OutJsonPin->MakeLinkTo(AsyncRequestNode->GetBodyPin());
+		CallFromStructFunctionNode->GetThenPin()->MakeLinkTo(AsyncRequestNode->GetExecPin());
+	}
+	else if (!GetBodyPin())
+	{
+		CompilerContext.MovePinLinksToIntermediate(*GetExecPin(), *AsyncRequestNode->GetExecPin());
+	}
+	else if (GetBodyPin())
+	{
+		CompilerContext.MessageLog.Error(
+			*LOCTEXT("Error", "Node @@ Attempted to supply a input body pin that was not a string or struct.").
+			ToString(), this);
+		BreakAllNodeLinks();
+	}
+	
+	UEdGraphPin* FuncExecPin = AsyncRequestNode->GetExecPin();
+	FuncExecPin->MakeLinkTo(GetExecPin());
+	
+	UEdGraphPin* FuncReturnValPin = AsyncRequestNode->FindPinChecked(TEXT("Response"));
+	UEdGraphPin* FuncConditionalPin = AsyncRequestNode->FindPinChecked(TEXT("bSuccessful"));
+	UEdGraphPin* FuncThenPin = AsyncRequestNode->FindPinChecked(TEXT("OnRequestComplete"));
+	UEdGraphPin* FuncOutHeaderPin = AsyncRequestNode->FindPinChecked(TEXT("OutHeader"));
+	
+	UK2Node_IfThenElse* RequestBranchNode = CompilerContext.SpawnIntermediateNode<UK2Node_IfThenElse>(this, SourceGraph);
+	RequestBranchNode->AllocateDefaultPins();
+
+	CompilerContext.MovePinLinksToIntermediate(*GetErrorPin(), *RequestBranchNode->GetElsePin());
+	CompilerContext.MovePinLinksToIntermediate(*GetOutHeaderPin(), *FuncOutHeaderPin);
+	
+	FuncThenPin->MakeLinkTo(RequestBranchNode->GetExecPin());
+	FuncConditionalPin->MakeLinkTo(RequestBranchNode->GetConditionPin());
+
+	/*
+	 *	We cache the PinType because in order for us to auto-convert the response data into the user defined struct
+	 *	We need to be able to create a JsonObject that can then be deserializated into the user defined struct
+	 **/
+	const FEdGraphPinType CachedOutputBodyPinType = GetOutBodyPin()->PinType;
+	if (CachedOutputBodyPinType.PinCategory != UEdGraphSchema_K2::PC_String)
+	{
+		GetOutBodyPin()->PinType.ResetToDefaults();
+		GetOutBodyPin()->PinType.PinCategory = UEdGraphSchema_K2::PC_String;
+	}
+	
+	// We check the cached pin type because the OutBodyPin should always be a string at this point
+	if (CachedOutputBodyPinType.PinCategory == UEdGraphSchema_K2::PC_String)
+	{
+		CompilerContext.MovePinLinksToIntermediate(*GetSuccessPin(), *RequestBranchNode->GetThenPin());
+		CompilerContext.MovePinLinksToIntermediate(*GetErrorPin(), *RequestBranchNode->GetElsePin());
+		CompilerContext.MovePinLinksToIntermediate(*GetOutBodyPin(), *FuncReturnValPin);
+	}
+	else if (CachedOutputBodyPinType.PinCategory == UEdGraphSchema_K2::PC_Struct)
+	{
+		const FName& FromStringFunctionName = GET_FUNCTION_NAME_CHECKED(UJsonBlueprintFunctionLibrary, FromString);
+		const FName& GetFieldFunctionName = GET_FUNCTION_NAME_CHECKED(UJsonBlueprintFunctionLibrary, GetField);
+		UClass* JsonLibraryClass = UJsonBlueprintFunctionLibrary::StaticClass();
+
+		UK2Node_CallFunction* CallFromStringFunctionNode = SourceGraph->CreateIntermediateNode<UK2Node_CallFunction>();
+		CallFromStringFunctionNode->FunctionReference.SetExternalMember(FromStringFunctionName, JsonLibraryClass);
+		CallFromStringFunctionNode->AllocateDefaultPins();
+		CompilerContext.MessageLog.NotifyIntermediateObjectCreation(CallFromStringFunctionNode, this);		
+
+		UEdGraphPin* JsonStringPin = CallFromStringFunctionNode->FindPinChecked(TEXT("JsonString"));
+		UEdGraphPin* JsonObjectPin = CallFromStringFunctionNode->FindPinChecked(TEXT("OutJsonObject"));
+		
+		UK2Node_CallFunction* CallGetFieldFunctionNode = SourceGraph->CreateIntermediateNode<UK2Node_CallFunction>();
+		CallGetFieldFunctionNode->FunctionReference.SetExternalMember(GetFieldFunctionName, JsonLibraryClass);
+		CallGetFieldFunctionNode->AllocateDefaultPins();
+		CompilerContext.MessageLog.NotifyIntermediateObjectCreation(CallGetFieldFunctionNode, this);
+
+		UEdGraphPin* ObjectWrapperPin = CallGetFieldFunctionNode->FindPinChecked(TEXT("JsonObject"));
+		UEdGraphPin* OutValuePin = CallGetFieldFunctionNode->FindPinChecked(TEXT("OutValue"));
+		UEdGraphPin* GetFieldConditionalPin = CallGetFieldFunctionNode->GetReturnValuePin();
+		
+		RequestBranchNode->GetThenPin()->MakeLinkTo(CallFromStringFunctionNode->GetExecPin());
+		CompilerContext.MovePinLinksToIntermediate(*GetErrorPin(), *RequestBranchNode->GetElsePin());
+		
+		FuncReturnValPin->MakeLinkTo(JsonStringPin);
+
+		CallFromStringFunctionNode->GetThenPin()->MakeLinkTo(CallGetFieldFunctionNode->GetExecPin());
+		JsonObjectPin->MakeLinkTo(ObjectWrapperPin);
+
+		UK2Node_IfThenElse* JsonBranchNode = CompilerContext.SpawnIntermediateNode<UK2Node_IfThenElse>(this, SourceGraph);
+		JsonBranchNode->AllocateDefaultPins();
+
+		CallGetFieldFunctionNode->GetThenPin()->MakeLinkTo(JsonBranchNode->GetExecPin());
+		GetFieldConditionalPin->MakeLinkTo(JsonBranchNode->GetConditionPin());
+
+		OutValuePin->PinType = CachedOutputBodyPinType;
+		GetOutBodyPin()->PinType = CachedOutputBodyPinType;
+
+		CompilerContext.MovePinLinksToIntermediate(*GetOutBodyPin(), *OutValuePin);
+		CompilerContext.MovePinLinksToIntermediate(*GetSuccessPin(), *JsonBranchNode->GetThenPin());
+	}
+	else
+	{
+		CompilerContext.MessageLog.Error(
+			*LOCTEXT("Error", "Node @@ Attempted to deserialize the Result Body into a @@. The supported types are String and Struct").
+			ToString(), this, *CachedOutputBodyPinType.PinCategory.ToString());
+		BreakAllNodeLinks();
+	}
+
+	BreakAllNodeLinks();
+}
+
+void UK2Node_HttpRequest::ReconstructNode()
+{
+	Super::ReconstructNode();
+
+	HandleBodyInputPin();
+	SyncBodyPinType(GetBodyPin());
+	SyncBodyPinType(GetOutBodyPin());
+}
+
+FName UK2Node_HttpRequest::GetCornerIcon() const
+{
+	return TEXT("Graph.Latent.LatentIcon");
+}
+
+void UK2Node_HttpRequest::GetMenuActions(FBlueprintActionDatabaseRegistrar& ActionRegistrar) const
+{
+	if (const UClass* ActionKey = GetClass();
+		ActionRegistrar.IsOpenForRegistration(ActionKey))
+	{
+		UBlueprintNodeSpawner* NodeSpawner = UBlueprintNodeSpawner::Create(GetClass());
+		checkf(NodeSpawner != nullptr, TEXT("Node spawner failed to create a valid Node"));
+
+		ActionRegistrar.AddBlueprintAction(ActionKey, NodeSpawner);
+	}
+}
+
+FText UK2Node_HttpRequest::GetMenuCategory() const
+{
+	return LOCTEXT("HttpRequest_Category", "Http");
+}
+
+bool UK2Node_HttpRequest::IsConnectionDisallowed(
+	const UEdGraphPin* MyPin,
+	const UEdGraphPin* OtherPin,
+	FString& OutReason) const
+{
+	if (MyPin == GetOutBodyPin())
+	{
+		if (OtherPin->PinType.PinCategory == UEdGraphSchema_K2::PC_Struct || OtherPin->PinType.PinCategory == UEdGraphSchema_K2::PC_String)
+		{
+			return false;
+		}
+
+		OutReason = LOCTEXT("HttpRequest_InvalidArgumentType", "Output Body may only be a string or struct").ToString();
+		return true;
+	}
+	
+	if (MyPin == GetBodyPin())
+	{
+		const FName& OtherPinCategory = OtherPin->PinType.PinCategory;
+		bool bIsValidType = false;
+
+		if (OtherPinCategory == UEdGraphSchema_K2::PC_String || OtherPinCategory == UEdGraphSchema_K2::PC_Struct)
+		{
+			bIsValidType = true;
+		}
+
+		if (!bIsValidType)
+		{
+			OutReason = LOCTEXT("HttpRequest_InvalidArgumentType", "Input Body may only be a string or struct").ToString();
+			return true;
+		}
+	}
+	
+	return Super::IsConnectionDisallowed(MyPin, OtherPin, OutReason);
+}
+
+void UK2Node_HttpRequest::SetPinToolTip(UEdGraphPin* const& InPin, const FText& ToolTip) const
+{
+	InPin->PinToolTip = UEdGraphSchema_K2::TypeToText(InPin->PinType).ToString();
+
+	const UEdGraphSchema_K2* const Schema = GetDefault<UEdGraphSchema_K2>();
+	InPin->PinToolTip += TEXT("");
+	InPin->PinToolTip += Schema->GetPinDisplayName(InPin).ToString();
+	InPin->PinToolTip += LINE_TERMINATOR + ToolTip.ToString();
+}
+
+void UK2Node_HttpRequest::HandleBodyInputPin()
+{
+	const UEdGraphPin* const EnumPin = FindPinChecked(HttpRequestLiterals::VerbPinName);
+
+	if (const int32 EnumIndex = GetVerbEnum()->GetIndexByName(*EnumPin->DefaultValue);
+		EnumIndex != INDEX_NONE
+		&& GetVerbEnum()->GetValueByIndex(EnumIndex) <= StaticCast<int32>(EHttpVerbs::Patch))
+	{
+		if (FindPin(HttpRequestLiterals::BodyPinName) == nullptr)
+		{
+			CreatePin(EGPD_Input, UEdGraphSchema_K2::PC_Wildcard, HttpRequestLiterals::BodyPinName);
+		}
+		
+		FBlueprintEditorUtils::MarkBlueprintAsStructurallyModified(GetBlueprint());
+		return;
+	}
+
+	if (UEdGraphPin* const BodyPin = FindPin(HttpRequestLiterals::BodyPinName);
+		BodyPin != nullptr)
+	{
+		RemovePin(BodyPin);
+		FBlueprintEditorUtils::MarkBlueprintAsStructurallyModified(GetBlueprint());
+	}
+}
+
+UEnum* UK2Node_HttpRequest::GetVerbEnum() const
+{
+	return VerbEnum ? VerbEnum.Get() : FindObjectChecked<UEnum>(ANY_PACKAGE, TEXT("EHttpVerbs"));
+}
+
+void UK2Node_HttpRequest::SyncBodyPinType(UEdGraphPin* const Pin) const
+{
+	// Early out for cases where the input body pin is nullptr
+	if (!Pin)
+	{
+		return;
+	}
+
+	bool bBodyTypeChanged = false;
+
+	if (Pin->LinkedTo.Num() == 0)
+	{
+		if (static const FEdGraphPinType WildcardType{
+			UEdGraphSchema_K2::PC_Wildcard,
+			NAME_None,
+			nullptr,
+			EPinContainerType::None,
+			false,
+			{}};
+		Pin->PinType != WildcardType)
+		{
+			Pin->PinType = WildcardType;
+			bBodyTypeChanged = true;
+		}
+	}
+	else
+	{
+		if (const FEdGraphPinType SourcePinType = Pin->LinkedTo[0]->PinType;
+			Pin->PinType != SourcePinType)
+		{
+			Pin->PinType = SourcePinType;
+			bBodyTypeChanged = true;
+		}
+	}
+
+	if (bBodyTypeChanged)
+	{
+		GetGraph()->NotifyGraphChanged();
+		if (UBlueprint* const BP = GetBlueprint(); !BP->bBeingCompiled)
+		{
+			FBlueprintEditorUtils::MarkBlueprintAsModified(BP);
+			BP->BroadcastChanged();
+		}
+	}
+}
+
+UEdGraphPin* UK2Node_HttpRequest::GetVerbPin() const
+{
+	UEdGraphPin* Pin = FindPinChecked(HttpRequestLiterals::VerbPinName);
+	check(Pin->Direction == EGPD_Input);
+	return Pin;
+}
+
+UEdGraphPin* UK2Node_HttpRequest::GetOutBodyPin() const
+{
+	UEdGraphPin* Pin = FindPinChecked(HttpRequestLiterals::OutputBodyPinName);
+	check(Pin->Direction == EGPD_Output);
+	return Pin;
+}
+
+UEdGraphPin* UK2Node_HttpRequest::GetBodyPin() const
+{
+	// The input body is dynamic. So there are cases where this will return nullptr
+	UEdGraphPin* Pin = FindPin(HttpRequestLiterals::BodyPinName);
+	check(Pin ? Pin->Direction == EGPD_Input : true);
+	return Pin;
+}
+
+UEdGraphPin* UK2Node_HttpRequest::GetHeaderPin() const
+{
+	UEdGraphPin* Pin = FindPinChecked(HttpRequestLiterals::HeaderPinName);
+	check(Pin->Direction == EGPD_Input);
+	return Pin;
+}
+
+UEdGraphPin* UK2Node_HttpRequest::GetOutHeaderPin() const
+{
+	UEdGraphPin* Pin = FindPinChecked(HttpRequestLiterals::OutHeaderPinName);
+	check(Pin->Direction == EGPD_Output);
+	return Pin;
+}
+
+UEdGraphPin* UK2Node_HttpRequest::GetUrlPin() const
+{
+	UEdGraphPin* Pin = FindPinChecked(HttpRequestLiterals::UrlPinName);
+	check(Pin->Direction == EGPD_Input);
+	return Pin;
+}
+
+UEdGraphPin* UK2Node_HttpRequest::GetSuccessPin() const
+{
+	UEdGraphPin* Pin = FindPinChecked(HttpRequestLiterals::SuccessPinName);
+	check(Pin->Direction == EGPD_Output);
+	return Pin;
+}
+
+UEdGraphPin* UK2Node_HttpRequest::GetErrorPin() const
+{
+	UEdGraphPin* Pin = FindPinChecked(HttpRequestLiterals::ErrorPinName);
+	check(Pin->Direction == EGPD_Output);
+	return Pin;
+}
+
+UEdGraphPin* UK2Node_HttpRequest::GetThenPin() const
+{
+	UEdGraphPin* Pin = FindPinChecked(UEdGraphSchema_K2::PN_Then);
+	check(Pin->Direction == EGPD_Output);
+	return Pin;
+}
+
+#undef LOCTEXT_NAMESPACE
diff --git a/Engine/Plugins/Web/HttpBlueprint/Source/HttpBlueprintGraph/Private/K2Node_HttpRequest.h b/Engine/Plugins/Web/HttpBlueprint/Source/HttpBlueprintGraph/Private/K2Node_HttpRequest.h
new file mode 100644
--- /dev/null
+++ b/Engine/Plugins/Web/HttpBlueprint/Source/HttpBlueprintGraph/Private/K2Node_HttpRequest.h
@@ -0,0 +1,59 @@
+// Copyright Epic Games, Inc. All Rights Reserved.
+
+#pragma once
+
+#include "CoreMinimal.h"
+#include "K2Node.h"
+
+#include "K2Node_HttpRequest.generated.h"
+
+UCLASS(MinimalAPI)
+class UK2Node_HttpRequest : public UK2Node
+{
+	GENERATED_BODY()
+
+public:
+	explicit UK2Node_HttpRequest(const FObjectInitializer& ObjectInitializer);
+	
+	/** Begin UEdGraphNode Interface */
+	virtual void AllocateDefaultPins() override;
+	virtual FText GetTooltipText() const override;
+	virtual FText GetNodeTitle(ENodeTitleType::Type TitleType) const override;
+	virtual bool IsCompatibleWithGraph(const UEdGraph* Graph) const override;
+	virtual void PinConnectionListChanged(UEdGraphPin* Pin) override;
+	virtual void PinDefaultValueChanged(UEdGraphPin* Pin) override;
+	virtual void PinTypeChanged(UEdGraphPin* Pin) override;
+	/** End UEdGraphNode Interface */
+
+	/** Begin UK2Node Interface */
+	virtual void ExpandNode(FKismetCompilerContext& CompilerContext, UEdGraph* SourceGraph) override;
+	virtual void ReconstructNode() override;
+	virtual FName GetCornerIcon() const override;
+	virtual void GetMenuActions(FBlueprintActionDatabaseRegistrar& ActionRegistrar) const override;
+	virtual bool NodeCausesStructuralBlueprintChange() const override { return true; }
+	virtual FText GetMenuCategory() const override;
+	virtual bool IsConnectionDisallowed(const UEdGraphPin* MyPin, const UEdGraphPin* OtherPin, FString& OutReason) const override;
+	/** Begin UK2Node Interface */
+
+protected:
+	void SetPinToolTip(UEdGraphPin* const& InPin, const FText& ToolTip) const;
+	
+	void HandleBodyInputPin();
+	UEnum* GetVerbEnum() const;
+
+	void SyncBodyPinType(UEdGraphPin* const Pin) const;
+
+	UEdGraphPin* GetVerbPin() const;
+	UEdGraphPin* GetOutBodyPin() const;
+	UEdGraphPin* GetBodyPin() const;
+	UEdGraphPin* GetHeaderPin() const;
+	UEdGraphPin* GetOutHeaderPin() const;
+	UEdGraphPin* GetUrlPin() const;
+	UEdGraphPin* GetSuccessPin() const;
+	UEdGraphPin* GetErrorPin() const;
+	UEdGraphPin* GetThenPin() const;
+	
+private:
+	UPROPERTY()
+	TObjectPtr<UEnum> VerbEnum;
+};
diff --git a/Engine/Plugins/Web/HttpBlueprint/Source/HttpBlueprintGraph/Private/K2Node_MakeRequestHeader.cpp b/Engine/Plugins/Web/HttpBlueprint/Source/HttpBlueprintGraph/Private/K2Node_MakeRequestHeader.cpp
new file mode 100644
--- /dev/null
+++ b/Engine/Plugins/Web/HttpBlueprint/Source/HttpBlueprintGraph/Private/K2Node_MakeRequestHeader.cpp
@@ -0,0 +1,496 @@
+// Copyright Epic Games, Inc. All Rights Reserved.
+
+#include "K2Node_MakeRequestHeader.h"
+
+#include "BlueprintActionDatabaseRegistrar.h"
+#include "BlueprintNodeSpawner.h"
+#include "EdGraphUtilities.h"
+#include "HttpBlueprintFunctionLibrary.h"
+#include "HttpBlueprintGraphLog.h"
+#include "HttpBlueprintTypes.h"
+#include "HttpHeader.h"
+#include "K2Node_CallFunction.h"
+#include "K2Node_MakeMap.h"
+#include "KismetCompiler.h"
+#include "ScopedTransaction.h"
+#include "ToolMenus.h"
+#include "Framework/Commands/UIAction.h"
+#include "Kismet2/BlueprintEditorUtils.h"
+
+#define LOCTEXT_NAMESPACE "K2Node_MakeRequestHeader"
+
+namespace UE::HttpBlueprint::Private
+{
+	namespace MakeRequestHeaderGlobals
+	{
+		static const FName OutputPinName(TEXT("Header"));
+		static const FName PresetPinName(TEXT("Preset"));
+		
+		static TArray<TPair<FString, FString>> Presets;
+	}
+	
+	// This exists so we can just store an array of the presets instead of a map.
+	// This allows us to avoid looping and comparing Enum values.
+	// Instead we can just use random access since the index is known
+	namespace FPresetPair
+	{
+		namespace Private
+		{
+			static FString GetPresetValueFromEnum(const ERequestPresets& InPresetEnum)
+			{
+				switch (InPresetEnum)
+				{
+					case ERequestPresets::Json:		return TEXT("application/json");
+					case ERequestPresets::Http:		return TEXT("text/html");
+					case ERequestPresets::Url:		return TEXT("application/x-www-form-urlencoded");
+					case ERequestPresets::Custom:	return TEXT("");
+					default:						return TEXT("");
+				}
+			}
+		}
+
+		inline namespace Public
+		{
+			static FString PresetArrayToString()
+			{
+				FString OutString = TEXT("Preset String");
+				for (int32 Index{0}; Index < MakeRequestHeaderGlobals::Presets.Num(); Index++)
+				{
+					OutString = FString::Printf(TEXT("%s%s: %s"), LINE_TERMINATOR,
+					                            *MakeRequestHeaderGlobals::Presets[Index].Key,
+					                            *MakeRequestHeaderGlobals::Presets[Index].Value);
+				}
+
+				return OutString;
+			}
+			
+			static FString GetHeaderValueForIndex(const int32& InIndex, const int32& InEnumIndex)
+			{
+				if (!MakeRequestHeaderGlobals::Presets.IsValidIndex(InIndex))
+				{
+					FFrame::KismetExecutionMessage(*FString::Printf(
+						TEXT("Preset array was not properly initialized. Can not make a request header. Index: %i -- Presets: %s"), InIndex, *PresetArrayToString()),
+						ELogVerbosity::Error);
+					return {};
+				}
+
+				const FString& HeaderValue = MakeRequestHeaderGlobals::Presets[InIndex].Value;
+				return HeaderValue.IsEmpty() ? Private::GetPresetValueFromEnum(
+						       StaticCast<ERequestPresets>(InEnumIndex)) : HeaderValue;
+			}
+
+			static FString GetHeaderKeyForIndex(const int32& InIndex)
+			{
+				if (!MakeRequestHeaderGlobals::Presets.IsValidIndex(InIndex))
+				{
+					FFrame::KismetExecutionMessage(*FString::Printf(
+					TEXT("Preset array was not properly initialized. Can not make a request header. Index: %i -- Presets: %s"), InIndex, *PresetArrayToString()),
+						ELogVerbosity::Error);
+					return {};
+				}
+				
+				return MakeRequestHeaderGlobals::Presets[InIndex].Key;
+			}
+
+			static int32 NumPresetHeaders()
+			{
+				return MakeRequestHeaderGlobals::Presets.Num();
+			}
+		}
+	};
+}
+
+UK2Node_MakeRequestHeader::UK2Node_MakeRequestHeader(const FObjectInitializer& ObjectInitializer)
+	: Super(ObjectInitializer)
+{
+	NumInputs = 1;
+	PresetEnum = FindObjectChecked<UEnum>(ANY_PACKAGE, TEXT("ERequestPresets"));
+
+	UE::HttpBlueprint::Private::MakeRequestHeaderGlobals::Presets =
+	{
+		{ TTuple<FString, FString>("Content-Type", "") },
+		{ TTuple<FString, FString>("Accepts", "") },
+		{ TTuple<FString, FString>("User-Agent", "X-UnrealEngine-Agent") },
+		{ TTuple<FString, FString>("Accept-Encoding", "identity") },
+	};
+}
+
+void UK2Node_MakeRequestHeader::AllocateDefaultPins()
+{
+	Super::AllocateDefaultPins();
+	
+	CreatePin(EGPD_Output, UEdGraphSchema_K2::PC_Struct, FHttpHeader::StaticStruct(), GetOutputPinName());
+
+	const UEdGraphSchema_K2* Schema = GetDefault<UEdGraphSchema_K2>();
+	UEdGraphPin* PresetPin = CreatePin(EGPD_Input,
+		UEdGraphSchema_K2::PC_Byte,
+		PresetEnum.Get(),
+		UE::HttpBlueprint::Private::MakeRequestHeaderGlobals::PresetPinName);
+	Schema->SetPinDefaultValueAtConstruction(PresetPin, PresetEnum->GetNameStringByIndex(PresetEnumIndex));
+
+	ConstructDefaultPinsForPreset();
+}
+
+FText UK2Node_MakeRequestHeader::GetTooltipText() const
+{
+	return LOCTEXT("MakeRequestHeader_Tooltip", "Creates a Header object that can be used to send Http Requests");
+}
+
+FText UK2Node_MakeRequestHeader::GetNodeTitle(ENodeTitleType::Type Title) const
+{
+	return LOCTEXT("MakeRequestHeader_Title", "Make Request Header");
+}
+
+void UK2Node_MakeRequestHeader::ValidateNodeDuringCompilation(FCompilerResultsLog& MessageLog) const
+{
+	Super::ValidateNodeDuringCompilation(MessageLog);
+}
+
+void UK2Node_MakeRequestHeader::PinDefaultValueChanged(UEdGraphPin* Pin)
+{
+	Super::PinDefaultValueChanged(Pin);
+
+	if (Pin->Direction == EGPD_Input
+		&& Pin->PinType.PinCategory == UEdGraphSchema_K2::PC_Byte)
+	{
+		ConstructDefaultPinsForPreset();
+	}
+}
+
+void UK2Node_MakeRequestHeader::GetNodeContextMenuActions(UToolMenu* Menu, UGraphNodeContextMenuContext* Context) const
+{
+	Super::GetNodeContextMenuActions(Menu, Context);
+
+	if (!Context->bIsDebugging)
+	{
+		FToolMenuSection& Section = Menu->AddSection("K2Node_MakeRequestHeader", LOCTEXT("MakeRequestHeader_ContextMenu", "Make Request Header"));
+		
+		if (Context->Pin)
+		{
+			if (Context->Pin->Direction == EGPD_Input && !Context->Pin->ParentPin)
+			{
+				Section.AddMenuEntry("RemovePin", LOCTEXT("MakeRequestHeader_RemovePin", "Remove header pin"),
+				                     LOCTEXT("MakeRequestHeader_RemovePin_Tooltip", "Remove this header element pin"),
+				                     FSlateIcon(),
+				                     FUIAction(FExecuteAction::CreateUObject(
+					                     const_cast<UK2Node_MakeRequestHeader*>(this),
+					                     &UK2Node_MakeRequestHeader::RemoveInputPin,
+					                     const_cast<UEdGraphPin*>(Context->Pin))));
+			}
+		}
+		else
+		{
+			Section.AddMenuEntry("AddPin", LOCTEXT("MakeRequestHeader_AddPin", "Add header pin"),
+			                     LOCTEXT("MakeRequestHeader_AddPin_Tooltip", "Add another header pin"),
+			                     FSlateIcon(),
+			                     FUIAction(FExecuteAction::CreateUObject(
+				                     const_cast<UK2Node_MakeRequestHeader*>(this),
+				                     &UK2Node_MakeRequestHeader::AddInputPin)));
+		}
+	}
+}
+
+void UK2Node_MakeRequestHeader::ExpandNode(FKismetCompilerContext& CompilerContext, UEdGraph* SourceGraph)
+{
+	Super::ExpandNode(CompilerContext, SourceGraph);
+	const FName& FunctionName = GET_FUNCTION_NAME_CHECKED(UHttpBlueprintFunctionLibrary, MakeRequestHeader);
+	UClass* FunctionClass = UHttpBlueprintFunctionLibrary::StaticClass();
+
+	UK2Node_CallFunction* CallFunctionNode = SourceGraph->CreateIntermediateNode<UK2Node_CallFunction>();
+	CallFunctionNode->FunctionReference.SetExternalMember(FunctionName, FunctionClass);
+	CallFunctionNode->AllocateDefaultPins();
+	CompilerContext.MessageLog.NotifyIntermediateObjectCreation(CallFunctionNode, this);
+
+	UK2Node_MakeMap* MakeMapNode = SourceGraph->CreateIntermediateNode<UK2Node_MakeMap>();
+	MakeMapNode->AllocateDefaultPins();
+	CompilerContext.MessageLog.NotifyIntermediateObjectCreation(MakeMapNode, this);
+
+	UEdGraphPin* MapOutPin = MakeMapNode->GetOutputPin();
+
+	UEdGraphPin* CallFunctionInputPin = CallFunctionNode->FindPinChecked(TEXT("Headers"));
+	MapOutPin->MakeLinkTo(CallFunctionInputPin);
+	MakeMapNode->PinConnectionListChanged(MapOutPin);
+	const UEdGraphSchema_K2* Schema = GetDefault<UEdGraphSchema_K2>();
+
+	if (PresetEnumIndex == StaticCast<int32>(ERequestPresets::Custom))
+	{
+		// Add input pins to the MakeMap node
+		// UK2Node_MakeMap::AddInputPin adds two pins per call which is why we only need to call it for Pins.Num / 2
+		for (int32 Index{0}; Index < (Pins.Num() / 2); Index++)
+		{
+			if (IsValidInputPin(Pins[Index]))
+			{
+				MakeMapNode->AddInputPin();
+			}
+		}
+	
+		// Skip the first index because it's the output pin
+		for (int32 Index{1}; Index < Pins.Num(); ++Index)
+		{
+			if (IsValidInputPin(Pins[Index]))
+			{
+				UEdGraphPin* MapPin = MakeMapNode->FindPinChecked(MakeMapNode->GetPinName(Index - 2));
+				MapPin->PinType = Pins[Index]->PinType;
+				CompilerContext.MovePinLinksToIntermediate(*Pins[Index], *MapPin);
+			}
+		}
+	}
+	else
+	{
+		if (PresetEnumIndex != StaticCast<int32>(ERequestPresets::Custom))
+		{
+			// MakeMapNode spawns one set of pins by default. So we only need to Spawn NumPins - 1
+			for (int32 Index{0}; Index < UE::HttpBlueprint::Private::FPresetPair::NumPresetHeaders() - 1; Index++)
+			{
+				MakeMapNode->AddInputPin();
+			}
+			
+			TArray<UEdGraphPin*> KeyPins, ValuePins;
+			MakeMapNode->GetKeyAndValuePins(KeyPins, ValuePins);
+			for (int32 Index{0}; Index < KeyPins.Num(); Index++)
+			{
+				UEdGraphPin* KeyPin = KeyPins[Index];
+				UEdGraphPin* ValuePin = ValuePins[Index];
+
+				Schema->TrySetDefaultValue(*KeyPin, UE::HttpBlueprint::Private::FPresetPair::GetHeaderKeyForIndex(Index));
+				Schema->TrySetDefaultValue(*ValuePin, UE::HttpBlueprint::Private::FPresetPair::GetHeaderValueForIndex(Index, PresetEnumIndex));
+			}
+		}
+	}
+
+	CompilerContext.MovePinLinksToIntermediate(*FindPinChecked(GetOutputPinName()),
+	                                           *CallFunctionNode->FindPinChecked(TEXT("OutHeader")));
+
+	BreakAllNodeLinks();
+}
+
+void UK2Node_MakeRequestHeader::GetMenuActions(FBlueprintActionDatabaseRegistrar& ActionRegistrar) const
+{
+	if (const UClass* ActionKey = GetClass();
+		ActionRegistrar.IsOpenForRegistration(ActionKey))
+	{
+		UBlueprintNodeSpawner* NodeSpawner = UBlueprintNodeSpawner::Create(GetClass());
+		checkf(NodeSpawner != nullptr, TEXT("Node spawner failed to create a valid Node"));
+
+		ActionRegistrar.AddBlueprintAction(ActionKey, NodeSpawner);
+	}
+}
+
+FText UK2Node_MakeRequestHeader::GetMenuCategory() const
+{
+	return LOCTEXT("MakeRequestHeader_Category", "Http");
+}
+
+bool UK2Node_MakeRequestHeader::NodeCausesStructuralBlueprintChange() const
+{
+	return true;
+}
+
+bool UK2Node_MakeRequestHeader::IsNodePure() const
+{
+	return true;
+}
+
+FName UK2Node_MakeRequestHeader::GetOutputPinName()
+{
+	return UE::HttpBlueprint::Private::MakeRequestHeaderGlobals::OutputPinName;
+}
+
+UEdGraphPin* UK2Node_MakeRequestHeader::GetOutputPin() const
+{
+	return FindPinChecked(GetOutputPinName());
+}
+
+bool UK2Node_MakeRequestHeader::IsValidInputPin(const UEdGraphPin* InputPin) const
+{
+	return InputPin->Direction == EGPD_Input
+		&& InputPin->PinName != UE::HttpBlueprint::Private::MakeRequestHeaderGlobals::PresetPinName
+		&& InputPin->PinType.PinCategory == UEdGraphSchema_K2::PC_String;
+}
+
+void UK2Node_MakeRequestHeader::ConstructDefaultPinsForPreset()
+{
+	const UEdGraphPin* PresetPin = FindPinChecked(UE::HttpBlueprint::Private::MakeRequestHeaderGlobals::PresetPinName);
+	if (const int32 EnumIndex = PresetEnum->GetIndexByName(*PresetPin->DefaultValue);
+		EnumIndex != INDEX_NONE)
+	{
+		PresetEnumIndex = EnumIndex;
+		if (PresetEnumIndex == StaticCast<int32>(ERequestPresets::Custom))
+		{
+			if (OptionalPins.Num() > 0)
+			{
+				const UEdGraphSchema_K2* Schema = GetDefault<UEdGraphSchema_K2>();
+				for (const FOptionalPin& OptionalPin : OptionalPins)
+				{
+					UEdGraphPin* NewPin = CreatePin(EGPD_Input, UEdGraphSchema_K2::PC_String, OptionalPin.PinName);
+					Schema->SetPinDefaultValueAtConstruction(NewPin, OptionalPin.PinDefaultValue);
+				}
+				NumInputs = Pins.Num() - 2;
+				FBlueprintEditorUtils::MarkBlueprintAsStructurallyModified(GetBlueprint());
+			}
+		}
+		else
+		{
+			// Remove all pins except for the preset and output pin
+			// Store all custom pins for reconstruction later
+			OptionalPins.Empty(Pins.Num() - 2);
+			for (int32 Index{2}; Index < Pins.Num(); ++Index)
+			{
+				UEdGraphPin* PinToRemove = Pins.IsValidIndex(Index) ? Pins[Index] : nullptr;
+				if (PinToRemove == nullptr)
+				{
+					UE_LOG(LogHttpBlueprintEditor, Warning, TEXT("Pin was invalid"));
+					break;
+				}
+				
+				FScopedTransaction Transaction(LOCTEXT("RemovePinTx", "RemovePin"));
+				Modify();
+
+				FOptionalPin NewOptionalPin;
+				NewOptionalPin.PinName = PinToRemove->PinName;
+				NewOptionalPin.PinDefaultValue = PinToRemove->DefaultValue;
+				OptionalPins.Add(NewOptionalPin);
+				
+				TFunction<void(UEdGraphPin*)> RemovePinLambda = [this, &RemovePinLambda](UEdGraphPin* PinToRemove)->void
+				{
+					check(PinToRemove);
+					for (int32 SubPinIndex = PinToRemove->SubPins.Num() - 1; SubPinIndex >= 0; --SubPinIndex)
+					{
+						RemovePinLambda(PinToRemove->SubPins[SubPinIndex]);
+					}
+
+					int32 PinRemovalIndex = INDEX_NONE;
+					if (Pins.Find(PinToRemove, PinRemovalIndex))
+					{
+						Pins.RemoveAt(PinRemovalIndex);
+						PinToRemove->MarkAsGarbage();
+					}
+				};
+				
+				RemovePinLambda(PinToRemove);
+				PinConnectionListChanged(PinToRemove);
+				--Index;
+			}
+
+			NumInputs = 0;
+			FBlueprintEditorUtils::MarkBlueprintAsStructurallyModified(GetBlueprint());
+		}
+	}
+}
+
+bool UK2Node_MakeRequestHeader::CanAddPin() const
+{
+	return StaticCast<ERequestPresets>(PresetEnumIndex) == ERequestPresets::Custom;
+}
+
+bool UK2Node_MakeRequestHeader::CanRemovePin(const UEdGraphPin* Pin) const
+{
+	return StaticCast<ERequestPresets>(PresetEnumIndex) == ERequestPresets::Custom && Pin->Direction == EGPD_Input &&
+		Pin->PinType.PinCategory == UEdGraphSchema_K2::PC_String;
+}
+
+void UK2Node_MakeRequestHeader::AddInputPin()
+{
+	FScopedTransaction Transaction(LOCTEXT("AddPinTx", "Add Pin"));
+	Modify();
+
+	++NumInputs;
+	CreatePin(EGPD_Input, UEdGraphSchema_K2::PC_String,
+	          *FString::Printf(TEXT("%s"), *GetPinName((NumInputs - 1) * 2).ToString()));
+
+	CreatePin(EGPD_Input, UEdGraphSchema_K2::PC_String,
+	          *FString::Printf(TEXT("%s"), *GetPinName((NumInputs - 1) * 2 + 1).ToString()));
+	
+	if (!GetBlueprint()->bBeingCompiled)
+	{
+		FBlueprintEditorUtils::MarkBlueprintAsStructurallyModified(GetBlueprint());
+		GetGraph()->NotifyGraphChanged();
+	}
+}
+
+void UK2Node_MakeRequestHeader::RemoveInputPin(UEdGraphPin* Pin)
+{
+	check(Pin);
+	check(Pin->Direction == EGPD_Input);
+	check(Pin->ParentPin == nullptr);
+	check(Pins.Contains(Pin));
+
+	FScopedTransaction Transaction(LOCTEXT("RemovePinTx", "RemovePin"));
+	Modify();
+
+	TFunction<void(UEdGraphPin*)> RemovePinLambda = [this, &RemovePinLambda](UEdGraphPin* PinToRemove)->void
+	{
+		check(PinToRemove);
+		for (int32 SubPinIndex = PinToRemove->SubPins.Num() - 1; SubPinIndex >= 0; --SubPinIndex)
+		{
+			RemovePinLambda(PinToRemove->SubPins[SubPinIndex]);
+		}
+
+		int32 PinRemovalIndex = INDEX_NONE;
+		if (Pins.Find(PinToRemove, PinRemovalIndex))
+		{
+			Pins.RemoveAt(PinRemovalIndex);
+			PinToRemove->MarkAsGarbage();
+		}
+	};
+
+	RemovePinLambda(Pin);
+	PinConnectionListChanged(Pin);
+
+	--NumInputs;
+	SyncPinNames();
+	FBlueprintEditorUtils::MarkBlueprintAsStructurallyModified(GetBlueprint());
+}
+
+void UK2Node_MakeRequestHeader::SyncPinNames()
+{
+	int32 CurrentNumParentPins = 0;
+	for (int32 PinIndex = 0; PinIndex < Pins.Num(); ++PinIndex)
+	{		
+		UEdGraphPin*& CurrentPin = Pins[PinIndex];
+		if (CurrentPin->Direction == EGPD_Input &&
+			CurrentPin->ParentPin == nullptr	&&
+			IsValidInputPin(Pins[PinIndex]))
+		{
+			const FName OldName = CurrentPin->PinName;
+			const FName ElementName = GetPinName(CurrentNumParentPins++);
+
+			CurrentPin->Modify();
+			CurrentPin->PinName = ElementName;
+
+			if (CurrentPin->SubPins.Num() > 0)
+			{
+				const FString OldNameStr = OldName.ToString();
+				const FString ElementNameStr = ElementName.ToString();
+				FString OldFriendlyName = OldNameStr;
+				FString ElementFriendlyName = ElementNameStr;
+
+				OldFriendlyName.InsertAt(1, " ");
+				ElementFriendlyName.InsertAt(1, " ");
+
+				for (UEdGraphPin* SubPin : CurrentPin->SubPins)
+				{
+					FString SubPinFriendlyName = SubPin->PinFriendlyName.ToString();
+					SubPinFriendlyName.ReplaceInline(*OldFriendlyName, *ElementFriendlyName);
+
+					SubPin->Modify();
+					SubPin->PinName = *SubPin->PinName.ToString().Replace(*OldNameStr, *ElementNameStr);
+					SubPin->PinFriendlyName = FText::FromString(SubPinFriendlyName);
+				}
+			}
+		}
+	}
+}
+
+FName UK2Node_MakeRequestHeader::GetPinName(const int32& PinIndex)
+{
+	const int32 PairIndex = PinIndex / 2;
+	if (PinIndex % 2 == 0)
+	{
+		return *FString::Printf(TEXT("Header %d"), PairIndex);
+	}
+
+	return *FString::Printf(TEXT("Value %d"), PairIndex);
+}
+
+#undef LOCTEXT_NAMESPACE
diff --git a/Engine/Plugins/Web/HttpBlueprint/Source/HttpBlueprintGraph/Private/K2Node_MakeRequestHeader.h b/Engine/Plugins/Web/HttpBlueprint/Source/HttpBlueprintGraph/Private/K2Node_MakeRequestHeader.h
new file mode 100644
--- /dev/null
+++ b/Engine/Plugins/Web/HttpBlueprint/Source/HttpBlueprintGraph/Private/K2Node_MakeRequestHeader.h
@@ -0,0 +1,80 @@
+// Copyright Epic Games, Inc. All Rights Reserved.
+
+#pragma once
+
+#include "CoreMinimal.h"
+#include "K2Node.h"
+#include "K2Node_AddPinInterface.h"
+
+#include "K2Node_MakeRequestHeader.generated.h"
+
+USTRUCT()
+struct FOptionalPin
+{
+	GENERATED_BODY()
+
+	UPROPERTY()
+	FName PinName;
+
+	UPROPERTY()
+	FString PinDefaultValue;
+};
+
+UCLASS(MinimalAPI)
+class UK2Node_MakeRequestHeader
+	: public UK2Node
+	, public IK2Node_AddPinInterface
+{
+	GENERATED_BODY()
+
+public:
+	UK2Node_MakeRequestHeader(const FObjectInitializer& ObjectInitializer);
+
+	/** Begin UEdGraphNode Interface */
+	virtual void AllocateDefaultPins() override;
+	virtual FText GetTooltipText() const override;
+	virtual FText GetNodeTitle(ENodeTitleType::Type TitleType) const override;
+	virtual void ValidateNodeDuringCompilation(FCompilerResultsLog& MessageLog) const override;
+	virtual void PinDefaultValueChanged(UEdGraphPin* Pin) override;
+	/** End UEdGraphNode Interface */
+
+	/** Begin UK2Node Interface */
+	virtual void GetNodeContextMenuActions(UToolMenu* Menu, UGraphNodeContextMenuContext* Context) const override;
+	virtual void ExpandNode(FKismetCompilerContext& CompilerContext, UEdGraph* SourceGraph) override;
+	virtual void GetMenuActions(FBlueprintActionDatabaseRegistrar& ActionRegistrar) const override;
+	virtual FText GetMenuCategory() const override;
+	virtual bool NodeCausesStructuralBlueprintChange() const override;
+	virtual bool IsNodePure() const override;
+	/** Begin UK2Node Interface */
+
+	/** Begin IK2Node_AddPinInterface */
+	virtual bool CanAddPin() const override;
+	virtual bool CanRemovePin(const UEdGraphPin* Pin) const override;
+	virtual void AddInputPin() override;
+	virtual void RemoveInputPin(UEdGraphPin* Pin) override;
+	/** End IK2Node_AddPinInterface */
+
+	static FName GetOutputPinName();
+
+	UEdGraphPin* GetOutputPin() const;
+
+protected:
+	UE_NODISCARD bool IsValidInputPin(const UEdGraphPin* InputPin) const;
+
+	void ConstructDefaultPinsForPreset();
+
+	void SyncPinNames();
+	static FName GetPinName(const int32& PinIndex);
+
+private:
+	TObjectPtr<UEnum> PresetEnum;
+
+	UPROPERTY()
+	int32 PresetEnumIndex;
+
+	UPROPERTY()
+	int32 NumInputs;
+
+	UPROPERTY()
+	TArray<FOptionalPin> OptionalPins;
+};
