Http BP non-unity fix (from preflight 623bc700ec68595f3b9ce6b0)

diff --git a/Engine/Plugins/Web/HttpBlueprint/Source/HttpBlueprintGraph/Private/K2Node_AsyncMakeRequestHeader.cpp b/Engine/Plugins/Web/HttpBlueprint/Source/HttpBlueprintGraph/Private/K2Node_AsyncMakeRequestHeader.cpp
--- a/Engine/Plugins/Web/HttpBlueprint/Source/HttpBlueprintGraph/Private/K2Node_AsyncMakeRequestHeader.cpp
+++ b/Engine/Plugins/Web/HttpBlueprint/Source/HttpBlueprintGraph/Private/K2Node_AsyncMakeRequestHeader.cpp
@@ -3,6 +3,7 @@
 #include "K2Node_AsyncMakeRequestHeader.h"
 
 #include "HttpRequestProxyObject.h"
+#include "EdGraph/EdGraphPin.h"
 
 UK2Node_AsyncMakeRequestHeader::UK2Node_AsyncMakeRequestHeader()
 {
