Fix error: inline function 'FHttpHeader::IsValid' is not defined
Use MoveTemp for some additional strings
#rnx

diff --git a/Engine/Plugins/Web/HttpBlueprint/Source/HttpBlueprint/Private/HttpBlueprintFunctionLibrary.cpp b/Engine/Plugins/Web/HttpBlueprint/Source/HttpBlueprint/Private/HttpBlueprintFunctionLibrary.cpp
--- a/Engine/Plugins/Web/HttpBlueprint/Source/HttpBlueprint/Private/HttpBlueprintFunctionLibrary.cpp
+++ b/Engine/Plugins/Web/HttpBlueprint/Source/HttpBlueprint/Private/HttpBlueprintFunctionLibrary.cpp
@@ -49,7 +49,7 @@ void UHttpBlueprintFunctionLibrary::AddHeader(FHttpHeader& HeaderObject, FString
 		return;
 	}
 	
-	TPair<FString, FString> NewPair(NewHeader, NewHeaderValue);
+	TPair<FString, FString> NewPair(MoveTemp(NewHeader), MoveTemp(NewHeaderValue));
 	HeaderObject.AddHeader(MoveTemp(NewPair));
 }
 
diff --git a/Engine/Plugins/Web/HttpBlueprint/Source/HttpBlueprint/Private/HttpHeader.cpp b/Engine/Plugins/Web/HttpBlueprint/Source/HttpBlueprint/Private/HttpHeader.cpp
--- a/Engine/Plugins/Web/HttpBlueprint/Source/HttpBlueprint/Private/HttpHeader.cpp
+++ b/Engine/Plugins/Web/HttpBlueprint/Source/HttpBlueprint/Private/HttpHeader.cpp
@@ -1,4 +1,4 @@
-﻿// Copyright Epic Games, Inc. All Rights Reserved.
+// Copyright Epic Games, Inc. All Rights Reserved.
 
 #include "HttpHeader.h"
 
@@ -39,16 +39,6 @@ TArray<FString> FHttpHeader::GetAllHeaders() const
 	return OutArray;
 }
 
-const TMap<FString, FString>& FHttpHeader::GetAllHeadersAsMap() const
-{
-	return Headers;
-}
-
-bool FHttpHeader::IsValid() const
-{
-	return Headers.Num() > 0;
-}
-
 void FHttpHeader::AssignHeadersToRequest(const TSharedRef<IHttpRequest> Request)
 {
 	for (const TPair<FString, FString>& HeaderPair : Headers)
diff --git a/Engine/Plugins/Web/HttpBlueprint/Source/HttpBlueprint/Public/HttpHeader.h b/Engine/Plugins/Web/HttpBlueprint/Source/HttpBlueprint/Public/HttpHeader.h
--- a/Engine/Plugins/Web/HttpBlueprint/Source/HttpBlueprint/Public/HttpHeader.h
+++ b/Engine/Plugins/Web/HttpBlueprint/Source/HttpBlueprint/Public/HttpHeader.h
@@ -1,4 +1,4 @@
-﻿// Copyright Epic Games, Inc. All Rights Reserved.
+// Copyright Epic Games, Inc. All Rights Reserved.
 
 #pragma once
 
@@ -43,12 +43,18 @@ struct HTTPBLUEPRINT_API FHttpHeader
 	/**
 	 * Returns a copy of the internal header map
 	 */
-	UE_NODISCARD inline const TMap<FString, FString>& GetAllHeadersAsMap() const;
+	UE_NODISCARD const TMap<FString, FString>& GetAllHeadersAsMap() const
+	{
+		return Headers;
+	}
 
 	/**
 	 * Returns if the header is valid or not. This is done by checking the map size
 	 */
-	UE_NODISCARD inline bool IsValid() const;
+	UE_NODISCARD bool IsValid() const
+	{
+		return Headers.Num() > 0;
+	}
 
 	/**
 	 * Convenience function to set all of the headers in the internal header map on the request object.
