HttpBlueprint uplugin corrections

diff --git a/Engine/Plugins/Web/HttpBlueprint/HttpBlueprint.uplugin b/Engine/Plugins/Web/HttpBlueprint/HttpBlueprint.uplugin
--- a/Engine/Plugins/Web/HttpBlueprint/HttpBlueprint.uplugin
+++ b/Engine/Plugins/Web/HttpBlueprint/HttpBlueprint.uplugin
@@ -1,15 +1,16 @@
 {
-	"FileVersion": 3,
+	"FileVersion": 1,
 	"Version": 1,
 	"VersionName": "1.0",
 	"FriendlyName": "HttpBlueprint",
 	"Description": "Allows for sending and receiving HTTP requests in Blueprint",
 	"Category": "Web",
-	"CreatedBy": "Epic Games",
-	"CreatedByURL": "",
+	"CreatedBy" : "Epic Games, Inc.",
+	"CreatedByURL" : "http://epicgames.com",
 	"DocsURL": "",
 	"MarketplaceURL": "",
 	"SupportURL": "",
+	"EnabledByDefault" : false,
 	"CanContainContent": true,
 	"IsBetaVersion": true,
 	"IsExperimentalVersion": false,
