Changed HttpBlueprint Plugin from Beta To Experimental #jira none #rb simon.therriault #preflight none

diff --git a/Engine/Plugins/Web/HttpBlueprint/HttpBlueprint.uplugin b/Engine/Plugins/Web/HttpBlueprint/HttpBlueprint.uplugin
--- a/Engine/Plugins/Web/HttpBlueprint/HttpBlueprint.uplugin
+++ b/Engine/Plugins/Web/HttpBlueprint/HttpBlueprint.uplugin
@@ -12,8 +12,8 @@
 	"SupportURL": "",
 	"EnabledByDefault" : false,
 	"CanContainContent": true,
-	"IsBetaVersion": true,
-	"IsExperimentalVersion": false,
+	"IsBetaVersion": false,
+	"IsExperimentalVersion": true,
 	"Installed": false,
 	"Modules": [
 		{
