HTTP MakeHeader node not saving pin connections

diff --git a/Engine/Plugins/Web/HttpBlueprint/Source/HttpBlueprintGraph/Private/K2Node_MakeRequestHeader.cpp b/Engine/Plugins/Web/HttpBlueprint/Source/HttpBlueprintGraph/Private/K2Node_MakeRequestHeader.cpp
--- a/Engine/Plugins/Web/HttpBlueprint/Source/HttpBlueprintGraph/Private/K2Node_MakeRequestHeader.cpp
+++ b/Engine/Plugins/Web/HttpBlueprint/Source/HttpBlueprintGraph/Private/K2Node_MakeRequestHeader.cpp
@@ -10,9 +10,9 @@
 #include "HttpHeader.h"
 #include "K2Node_CallFunction.h"
 #include "K2Node_MakeMap.h"
+#include "Kismet2/BlueprintEditorUtils.h"
 #include "KismetCompiler.h"
 #include "ScopedTransaction.h"
-#include "Kismet2/BlueprintEditorUtils.h"
 #include "ToolMenu.h"
 #include "ToolMenuSection.h"
 
@@ -116,8 +116,6 @@ UK2Node_MakeRequestHeader::UK2Node_MakeRequestHeader(const FObjectInitializer& O
 
 void UK2Node_MakeRequestHeader::AllocateDefaultPins()
 {
-	Super::AllocateDefaultPins();
-	
 	CreatePin(EGPD_Output, UEdGraphSchema_K2::PC_Struct, FHttpHeader::StaticStruct(), GetOutputPinName());
 
 	const UEdGraphSchema_K2* Schema = GetDefault<UEdGraphSchema_K2>();
@@ -127,12 +125,16 @@ void UK2Node_MakeRequestHeader::AllocateDefaultPins()
 		UE::HttpBlueprint::Private::MakeRequestHeaderGlobals::PresetPinName);
 	Schema->SetPinDefaultValueAtConstruction(PresetPin, PresetEnum->GetNameStringByIndex(PresetEnumIndex));
 
-	ConstructDefaultPinsForPreset();
-}
+	// Create the input pins to create the container from
+	for (int32 InputIndex = 0; InputIndex < NumInputs; ++InputIndex)
+	{
+		CreatePin(EGPD_Input, UEdGraphSchema_K2::PC_String,*FString::Printf(TEXT("%s"), *GetPinName(InputIndex * 2).ToString()));
+		CreatePin(EGPD_Input, UEdGraphSchema_K2::PC_String, *FString::Printf(TEXT("%s"), *GetPinName(InputIndex * 2 + 1).ToString()));
+	}
 
-FText UK2Node_MakeRequestHeader::GetTooltipText() const
-{
-	return LOCTEXT("MakeRequestHeader_Tooltip", "Creates a Header object that can be used to send Http Requests");
+	SyncPinNames();
+
+	ConstructDefaultPinsForPreset();
 }
 
 FText UK2Node_MakeRequestHeader::GetNodeTitle(ENodeTitleType::Type Title) const
@@ -140,9 +142,9 @@ FText UK2Node_MakeRequestHeader::GetNodeTitle(ENodeTitleType::Type Title) const
 	return LOCTEXT("MakeRequestHeader_Title", "Make Request Header");
 }
 
-void UK2Node_MakeRequestHeader::ValidateNodeDuringCompilation(FCompilerResultsLog& MessageLog) const
+FText UK2Node_MakeRequestHeader::GetTooltipText() const
 {
-	Super::ValidateNodeDuringCompilation(MessageLog);
+	return LOCTEXT("MakeRequestHeader_Tooltip", "Creates a Header object that can be used to send Http Requests");
 }
 
 void UK2Node_MakeRequestHeader::PinDefaultValueChanged(UEdGraphPin* Pin)
@@ -257,8 +259,13 @@ void UK2Node_MakeRequestHeader::ExpandNode(FKismetCompilerContext& CompilerConte
 		}
 	}
 
-	CompilerContext.MovePinLinksToIntermediate(*FindPinChecked(GetOutputPinName()),
-	                                           *CallFunctionNode->FindPinChecked(TEXT("OutHeader")));
+	if (UEdGraphPin* OutputPin = FindPin(GetOutputPinName()))
+	{
+		if (UEdGraphPin* HeaderOutputPin = CallFunctionNode->FindPin(TEXT("OutHeader")))
+		{
+			CompilerContext.MovePinLinksToIntermediate(*OutputPin, *HeaderOutputPin);			
+		}
+	}
 
 	BreakAllNodeLinks();
 }
@@ -322,6 +329,10 @@ void UK2Node_MakeRequestHeader::ConstructDefaultPinsForPreset()
 				for (const FOptionalPin& OptionalPin : OptionalPins)
 				{
 					UEdGraphPin* NewPin = CreatePin(EGPD_Input, UEdGraphSchema_K2::PC_String, OptionalPin.PinName);
+					if (UEdGraphPin* LinkedToPin = OptionalPin.LinkedTo.Get())
+					{
+						NewPin->MakeLinkTo(LinkedToPin);						
+					}					
 					Schema->SetPinDefaultValueAtConstruction(NewPin, OptionalPin.PinDefaultValue);
 				}
 				NumInputs = Pins.Num() - 2;
@@ -346,6 +357,10 @@ void UK2Node_MakeRequestHeader::ConstructDefaultPinsForPreset()
 				Modify();
 
 				FOptionalPin NewOptionalPin;
+				if (!PinToRemove->LinkedTo.IsEmpty())
+				{
+					NewOptionalPin.LinkedTo = PinToRemove->LinkedTo[0];
+				}
 				NewOptionalPin.PinName = PinToRemove->PinName;
 				NewOptionalPin.PinDefaultValue = PinToRemove->DefaultValue;
 				OptionalPins.Add(NewOptionalPin);
@@ -417,10 +432,11 @@ void UK2Node_MakeRequestHeader::RemoveInputPin(UEdGraphPin* Pin)
 	FScopedTransaction Transaction(LOCTEXT("RemovePinTx", "RemovePin"));
 	Modify();
 
-	TFunction<void(UEdGraphPin*)> RemovePinLambda = [this, &RemovePinLambda](UEdGraphPin* PinToRemove)->void
+	TFunction<void(UEdGraphPin*)> RemovePinLambda = [this, &RemovePinLambda](UEdGraphPin* PinToRemove)
 	{
 		check(PinToRemove);
-		for (int32 SubPinIndex = PinToRemove->SubPins.Num() - 1; SubPinIndex >= 0; --SubPinIndex)
+		
+		for (int32 SubPinIndex = PinToRemove->SubPins.Num()-1; SubPinIndex >= 0; --SubPinIndex)
 		{
 			RemovePinLambda(PinToRemove->SubPins[SubPinIndex]);
 		}
@@ -433,6 +449,21 @@ void UK2Node_MakeRequestHeader::RemoveInputPin(UEdGraphPin* Pin)
 		}
 	};
 
+	TArray<UEdGraphPin*> KeyPins;
+	TArray<UEdGraphPin*> ValuePins;
+	GetKeyAndValuePins(KeyPins, ValuePins);
+
+	int32 PinIndex = INDEX_NONE;
+	if (ValuePins.Find(Pin, PinIndex))
+	{
+		RemovePinLambda(KeyPins[PinIndex]);
+	}
+	else
+	{
+		verify(KeyPins.Find(Pin, PinIndex));
+		RemovePinLambda(ValuePins[PinIndex]);
+	}
+
 	RemovePinLambda(Pin);
 	PinConnectionListChanged(Pin);
 
@@ -492,4 +523,26 @@ FName UK2Node_MakeRequestHeader::GetPinName(const int32& PinIndex)
 	return *FString::Printf(TEXT("Value %d"), PairIndex);
 }
 
+void UK2Node_MakeRequestHeader::GetKeyAndValuePins(TArray<UEdGraphPin*>& KeyPins, TArray<UEdGraphPin*>& ValuePins) const
+{
+	// skip the first two (input, output) pins
+	for (int32 PinIndex = 2; PinIndex < Pins.Num(); ++PinIndex)
+	{
+		UEdGraphPin* CurrentPin = Pins[PinIndex];
+		if (CurrentPin->Direction == EGPD_Input && CurrentPin->ParentPin == nullptr)
+		{
+			// Key/Value pins alternate so if this is a map and the counts are even then this is a key
+			if (KeyPins.Num() == ValuePins.Num())
+			{
+				KeyPins.Add(CurrentPin);
+			}
+			else
+			{
+				ValuePins.Add(CurrentPin);
+			}
+		}
+	}
+	check(KeyPins.Num() == ValuePins.Num());
+}
+
 #undef LOCTEXT_NAMESPACE
diff --git a/Engine/Plugins/Web/HttpBlueprint/Source/HttpBlueprintGraph/Private/K2Node_MakeRequestHeader.h b/Engine/Plugins/Web/HttpBlueprint/Source/HttpBlueprintGraph/Private/K2Node_MakeRequestHeader.h
--- a/Engine/Plugins/Web/HttpBlueprint/Source/HttpBlueprintGraph/Private/K2Node_MakeRequestHeader.h
+++ b/Engine/Plugins/Web/HttpBlueprint/Source/HttpBlueprintGraph/Private/K2Node_MakeRequestHeader.h
@@ -7,10 +7,11 @@
 
 #include "K2Node_MakeRequestHeader.generated.h"
 
+namespace ENodeTitleType { enum Type : int; }
 class FCompilerResultsLog;
+struct FEdGraphPinReference;
 class UGraphNodeContextMenuContext;
 class UToolMenu;
-namespace ENodeTitleType { enum Type : int; }
 
 USTRUCT()
 struct FOptionalPin
@@ -22,6 +23,10 @@ struct FOptionalPin
 
 	UPROPERTY()
 	FString PinDefaultValue;
+
+	/** Optional, use if something was previously linked to the pin. */
+	UPROPERTY()
+	FEdGraphPinReference LinkedTo;
 };
 
 UCLASS(MinimalAPI)
@@ -36,9 +41,8 @@ public:
 
 	/** Begin UEdGraphNode Interface */
 	virtual void AllocateDefaultPins() override;
-	virtual FText GetTooltipText() const override;
 	virtual FText GetNodeTitle(ENodeTitleType::Type TitleType) const override;
-	virtual void ValidateNodeDuringCompilation(FCompilerResultsLog& MessageLog) const override;
+	virtual FText GetTooltipText() const override;
 	virtual void PinDefaultValueChanged(UEdGraphPin* Pin) override;
 	/** End UEdGraphNode Interface */
 
@@ -69,6 +73,7 @@ protected:
 
 	void SyncPinNames();
 	static FName GetPinName(const int32& PinIndex);
+	virtual void GetKeyAndValuePins(TArray<UEdGraphPin*>& KeyPins, TArray<UEdGraphPin*>& ValuePins) const;
 
 private:
 	TObjectPtr<UEnum> PresetEnum;
