[HttpBlueprint] Fixed MakeHttpHeader node

diff --git a/Engine/Plugins/Web/HttpBlueprint/Source/HttpBlueprint/Internal/HttpRequestProxyObject.h b/Engine/Plugins/Web/HttpBlueprint/Source/HttpBlueprint/Internal/HttpRequestProxyObject.h
new file mode 100644
--- /dev/null
+++ b/Engine/Plugins/Web/HttpBlueprint/Source/HttpBlueprint/Internal/HttpRequestProxyObject.h
@@ -0,0 +1,35 @@
+// Copyright Epic Games, Inc. All Rights Reserved.
+
+#pragma once
+
+#include "HttpHeader.h"
+#include "Interfaces/IHttpRequest.h"
+
+#include "HttpRequestProxyObject.generated.h"
+
+DECLARE_DYNAMIC_MULTICAST_DELEGATE_ThreeParams(FOnRequestComplete, FString, Response, bool, bSuccessful, FHttpHeader, OutHeader);
+
+UCLASS(MinimalAPI)
+class UHttpRequestProxyObject : public UObject
+{
+	GENERATED_BODY()
+
+public:
+	UPROPERTY(BlueprintAssignable)
+	FOnRequestComplete OnRequestComplete;
+
+	UFUNCTION(BlueprintCallable, Meta = (BlueprintInternalUseOnly, DisplayName = "Http Request"), Category = "Http")
+	static UHttpRequestProxyObject* CreateHttpRequestProxyObject(
+		const FString& InUrl,
+		const FString& InVerb,
+		FHttpHeader InHeader,
+		const FString& InBody);
+
+protected:
+	void ProcessRequest(const FString& InUrl, const FString& InVerb, FHttpHeader&& InHeader, const FString& InBody);
+	void ProcessComplete(FHttpRequestPtr InRequest, FHttpResponsePtr InResponse, bool bInSuccessful);
+
+private:
+	UPROPERTY()
+	FHttpHeader CachedHeader;
+};
diff --git a/Engine/Plugins/Web/HttpBlueprint/Source/HttpBlueprint/Private/HttpHeader.cpp b/Engine/Plugins/Web/HttpBlueprint/Source/HttpBlueprint/Private/HttpHeader.cpp
--- a/Engine/Plugins/Web/HttpBlueprint/Source/HttpBlueprint/Private/HttpHeader.cpp
+++ b/Engine/Plugins/Web/HttpBlueprint/Source/HttpBlueprint/Private/HttpHeader.cpp
@@ -2,7 +2,9 @@
 
 #include "HttpHeader.h"
 
+#include "HttpBlueprintTypes.h"
 #include "Interfaces/IHttpRequest.h"
+#include "Interfaces/IHttpResponse.h"
 
 FHttpHeader FHttpHeader::SetHeaders(const TMap<FString, FString>& NewHeaders)
 {
@@ -39,8 +41,21 @@ TArray<FString> FHttpHeader::GetAllHeaders() const
 	return OutArray;
 }
 
-void FHttpHeader::AssignHeadersToRequest(const TSharedRef<IHttpRequest> Request)
+void FHttpHeader::AssignHeadersToRequest(const TSharedRef<IHttpRequest>& Request)
 {
+	// Add required/expected headers by default - they'll be overridden if specified below
+	// Even if the values are wrong, the Http Request will notify the user of this in it's response
+	const FString RequestVerb = Request->GetVerb();
+	if (RequestVerb.Equals(TEXT("POST"), ESearchCase::IgnoreCase)
+		|| RequestVerb.Equals(TEXT("PUT"), ESearchCase::IgnoreCase)
+		|| RequestVerb.Equals(TEXT("PATCH"), ESearchCase::IgnoreCase))
+	{
+		const FString InferredContentType = Request->GetURL().Contains(TEXT("?"))
+												? TEXT("application/x-www-form-urlencoded") // The presence of ? in a url probably means it's parameter based
+												: TEXT("application/json"); // otherwise treat as having a json payload
+		Request->SetHeader(TEXT("Content-Type"), InferredContentType);
+	}
+
 	for (const TPair<FString, FString>& HeaderPair : Headers)
 	{
 		Request->SetHeader(HeaderPair.Key, HeaderPair.Value);
diff --git a/Engine/Plugins/Web/HttpBlueprint/Source/HttpBlueprint/Private/HttpRequestProxyObject.cpp b/Engine/Plugins/Web/HttpBlueprint/Source/HttpBlueprint/Private/HttpRequestProxyObject.cpp
new file mode 100644
--- /dev/null
+++ b/Engine/Plugins/Web/HttpBlueprint/Source/HttpBlueprint/Private/HttpRequestProxyObject.cpp
@@ -0,0 +1,41 @@
+// Copyright Epic Games, Inc. All Rights Reserved.
+
+#include "HttpRequestProxyObject.h"
+
+#include "HttpModule.h"
+#include "Interfaces/IHttpResponse.h"
+
+UHttpRequestProxyObject* UHttpRequestProxyObject::CreateHttpRequestProxyObject(
+	const FString& InUrl,
+	const FString& InVerb,
+	FHttpHeader InHeader,
+	const FString& InBody)
+{
+	UHttpRequestProxyObject* const Proxy = NewObject<UHttpRequestProxyObject>();
+	Proxy->SetFlags(RF_StrongRefOnFrame);
+	Proxy->CachedHeader = InHeader;
+	Proxy->ProcessRequest(InUrl, InVerb, MoveTemp(InHeader), InBody);
+	return Proxy;
+}
+
+void UHttpRequestProxyObject::ProcessRequest(
+	const FString& InUrl,
+	const FString& InVerb,
+	FHttpHeader&& InHeader,
+	const FString& InBody)
+{
+	FHttpModule& HttpModule = FHttpModule::Get();
+	const TSharedRef<IHttpRequest> Request = HttpModule.CreateRequest();
+	Request->SetURL(InUrl);
+	Request->SetVerb(InVerb);
+	Request->SetContentAsString(InBody);
+	InHeader.AssignHeadersToRequest(Request);
+	Request->ProcessRequest();
+
+	Request->OnProcessRequestComplete().BindUObject(this, &ThisClass::ProcessComplete);
+}
+
+void UHttpRequestProxyObject::ProcessComplete(FHttpRequestPtr InRequest, FHttpResponsePtr InResponse, bool bInSuccessful)
+{
+	OnRequestComplete.Broadcast(InResponse->GetContentAsString(), bInSuccessful, MoveTemp(CachedHeader));
+}
diff --git a/Engine/Plugins/Web/HttpBlueprint/Source/HttpBlueprint/Public/HttpHeader.h b/Engine/Plugins/Web/HttpBlueprint/Source/HttpBlueprint/Public/HttpHeader.h
--- a/Engine/Plugins/Web/HttpBlueprint/Source/HttpBlueprint/Public/HttpHeader.h
+++ b/Engine/Plugins/Web/HttpBlueprint/Source/HttpBlueprint/Public/HttpHeader.h
@@ -59,11 +59,10 @@ struct HTTPBLUEPRINT_API FHttpHeader
 	/**
 	 * Convenience function to set all of the headers in the internal header map on the request object.
 	 */
-	void AssignHeadersToRequest(TSharedRef<class IHttpRequest> Request);
+	void AssignHeadersToRequest(const TSharedRef<class IHttpRequest>& Request);
 	
 private:
-
+	/** Header key/value pairs. */
 	UPROPERTY()
-	TMap<FString, FString> Headers;
-	
+	TMap<FString, FString> Headers;	
 };
diff --git a/Engine/Plugins/Web/HttpBlueprint/Source/HttpBlueprintGraph/Private/HttpRequestProxyObject.cpp b/Engine/Plugins/Web/HttpBlueprint/Source/HttpBlueprintGraph/Private/HttpRequestProxyObject.cpp
deleted file mode 100644
diff --git a/Engine/Plugins/Web/HttpBlueprint/Source/HttpBlueprintGraph/Private/HttpRequestProxyObject.h b/Engine/Plugins/Web/HttpBlueprint/Source/HttpBlueprintGraph/Private/HttpRequestProxyObject.h
deleted file mode 100644
diff --git a/Engine/Plugins/Web/HttpBlueprint/Source/HttpBlueprintGraph/Private/K2Node_MakeRequestHeader.cpp b/Engine/Plugins/Web/HttpBlueprint/Source/HttpBlueprintGraph/Private/K2Node_MakeRequestHeader.cpp
--- a/Engine/Plugins/Web/HttpBlueprint/Source/HttpBlueprintGraph/Private/K2Node_MakeRequestHeader.cpp
+++ b/Engine/Plugins/Web/HttpBlueprint/Source/HttpBlueprintGraph/Private/K2Node_MakeRequestHeader.cpp
@@ -24,10 +24,74 @@ namespace UE::HttpBlueprint::Private
 	{
 		static const FName OutputPinName(TEXT("Header"));
 		static const FName PresetPinName(TEXT("Preset"));
+
+		static const FName InputKeyPrefix(TEXT("Header"));
+		static const FName InputValuePrefix(TEXT("Value"));
 		
 		static TArray<TPair<FString, FString>> Presets;
 	}
+
+	namespace Pin
+	{
+		UE_NODISCARD static FString MakePinName(int32 PinIndex)
+		{
+			const int32 PairIndex = PinIndex / 2;
+			if (PinIndex % 2 == 0)
+			{
+				return *FString::Printf(TEXT("%s %d"), *UE::HttpBlueprint::Private::MakeRequestHeaderGlobals::InputKeyPrefix.ToString(), PairIndex);
+			}
+
+			return *FString::Printf(TEXT("%s %d"), *UE::HttpBlueprint::Private::MakeRequestHeaderGlobals::InputValuePrefix.ToString(), PairIndex);
+		}
 	
+		UE_NODISCARD static bool IsValidInputPin(const UEdGraphPin* InputPin)
+		{
+			return InputPin
+				&& InputPin->Direction == EGPD_Input
+				&& InputPin->PinName != UE::HttpBlueprint::Private::MakeRequestHeaderGlobals::PresetPinName
+				&& InputPin->PinType.PinCategory == UEdGraphSchema_K2::PC_String;
+		}
+
+		/** Returns true if the Pin has any connection or an inline, non-default value. */
+		static bool PinHasConnectionOrValue(const UEdGraphPin* Pin)
+		{
+			return Pin->HasAnyConnections() || !Pin->DefaultValue.IsEmpty();
+		}
+
+		/** Returns true if one or more pins match the condition. */
+		static TArray<UEdGraphPin*> CopyKeyPinsIf(
+			TArrayView<UEdGraphPin*> From,
+			TUniqueFunction<bool(const UEdGraphPin*)>&& PredicateFunc)
+		{
+			TArray<UEdGraphPin*> To;
+			To.Reserve(From.Num());
+
+			for (int32 Index{0}; Index < From.Num(); Index += 2)
+			{
+				UEdGraphPin* Pin = From[Index];
+				if (UE::HttpBlueprint::Private::Pin::IsValidInputPin(Pin)
+					&& PredicateFunc(Pin))
+				{
+					To.Add(Pin);					
+				}
+			}
+
+			return To;
+		}
+
+		/** A Key Pin is considered valid if it has an input or non-default value. */
+		static TArray<UEdGraphPin*> GetValidKeyPins(TArrayView<UEdGraphPin*> Pins)
+		{
+			TArray<UEdGraphPin*> EmptyKeyPins = CopyKeyPinsIf(Pins,
+				[](const UEdGraphPin* KeyPin)
+				{
+					return PinHasConnectionOrValue(KeyPin);
+				});
+
+			return EmptyKeyPins;
+		}
+	}
+
 	// This exists so we can just store an array of the presets instead of a map.
 	// This allows us to avoid looping and comparing Enum values.
 	// Instead we can just use random access since the index is known
@@ -63,7 +127,7 @@ namespace UE::HttpBlueprint::Private
 				return OutString;
 			}
 			
-			static FString GetHeaderValueForIndex(const int32& InIndex, const int32& InEnumIndex)
+			static FString GetHeaderValueForIndex(int32 InIndex, int32 InEnumIndex)
 			{
 				if (!MakeRequestHeaderGlobals::Presets.IsValidIndex(InIndex))
 				{
@@ -78,7 +142,7 @@ namespace UE::HttpBlueprint::Private
 						       StaticCast<ERequestPresets>(InEnumIndex)) : HeaderValue;
 			}
 
-			static FString GetHeaderKeyForIndex(const int32& InIndex)
+			static FString GetHeaderKeyForIndex(int32 InIndex)
 			{
 				if (!MakeRequestHeaderGlobals::Presets.IsValidIndex(InIndex))
 				{
@@ -96,7 +160,7 @@ namespace UE::HttpBlueprint::Private
 				return MakeRequestHeaderGlobals::Presets.Num();
 			}
 		}
-	};
+	}
 }
 
 UK2Node_MakeRequestHeader::UK2Node_MakeRequestHeader(const FObjectInitializer& ObjectInitializer)
@@ -114,12 +178,21 @@ UK2Node_MakeRequestHeader::UK2Node_MakeRequestHeader(const FObjectInitializer& O
 	};
 }
 
+void UK2Node_MakeRequestHeader::ReallocatePinsDuringReconstruction(TArray<UEdGraphPin*>& OldPins)
+{
+	// @note: Important to do this here for default pin allocation to work (which is called below)
+	NumInputs = FMath::Max(1, (OldPins.Num() - 2) / 2);
+	
+	Super::ReallocatePinsDuringReconstruction(OldPins);
+}
+
 void UK2Node_MakeRequestHeader::AllocateDefaultPins()
 {
-	CreatePin(EGPD_Output, UEdGraphSchema_K2::PC_Struct, FHttpHeader::StaticStruct(), GetOutputPinName());
+	CreatePin(EGPD_Output, UEdGraphSchema_K2::PC_Struct, FHttpHeader::StaticStruct(), UE::HttpBlueprint::Private::MakeRequestHeaderGlobals::OutputPinName);
 
 	const UEdGraphSchema_K2* Schema = GetDefault<UEdGraphSchema_K2>();
-	UEdGraphPin* PresetPin = CreatePin(EGPD_Input,
+	UEdGraphPin* PresetPin = CreatePin(
+		EGPD_Input,
 		UEdGraphSchema_K2::PC_Byte,
 		PresetEnum.Get(),
 		UE::HttpBlueprint::Private::MakeRequestHeaderGlobals::PresetPinName);
@@ -128,8 +201,11 @@ void UK2Node_MakeRequestHeader::AllocateDefaultPins()
 	// Create the input pins to create the container from
 	for (int32 InputIndex = 0; InputIndex < NumInputs; ++InputIndex)
 	{
-		CreatePin(EGPD_Input, UEdGraphSchema_K2::PC_String,*FString::Printf(TEXT("%s"), *GetPinName(InputIndex * 2).ToString()));
-		CreatePin(EGPD_Input, UEdGraphSchema_K2::PC_String, *FString::Printf(TEXT("%s"), *GetPinName(InputIndex * 2 + 1).ToString()));
+		FString KeyPinName = UE::HttpBlueprint::Private::Pin::MakePinName(InputIndex * 2);
+		CreatePin(EGPD_Input, UEdGraphSchema_K2::PC_String,*KeyPinName);
+		
+		FString ValuePinName = UE::HttpBlueprint::Private::Pin::MakePinName(InputIndex * 2 + 1);
+		CreatePin(EGPD_Input, UEdGraphSchema_K2::PC_String, *ValuePinName);
 	}
 
 	SyncPinNames();
@@ -194,15 +270,18 @@ void UK2Node_MakeRequestHeader::GetNodeContextMenuActions(UToolMenu* Menu, UGrap
 void UK2Node_MakeRequestHeader::ExpandNode(FKismetCompilerContext& CompilerContext, UEdGraph* SourceGraph)
 {
 	Super::ExpandNode(CompilerContext, SourceGraph);
+
+	const UEdGraphSchema_K2* Schema = CompilerContext.GetSchema();
+	
 	const FName& FunctionName = GET_FUNCTION_NAME_CHECKED(UHttpBlueprintFunctionLibrary, MakeRequestHeader);
 	UClass* FunctionClass = UHttpBlueprintFunctionLibrary::StaticClass();
 
-	UK2Node_CallFunction* CallFunctionNode = SourceGraph->CreateIntermediateNode<UK2Node_CallFunction>();
+	UK2Node_CallFunction* CallFunctionNode = CompilerContext.SpawnIntermediateNode<UK2Node_CallFunction>(this, SourceGraph);
 	CallFunctionNode->FunctionReference.SetExternalMember(FunctionName, FunctionClass);
 	CallFunctionNode->AllocateDefaultPins();
 	CompilerContext.MessageLog.NotifyIntermediateObjectCreation(CallFunctionNode, this);
 
-	UK2Node_MakeMap* MakeMapNode = SourceGraph->CreateIntermediateNode<UK2Node_MakeMap>();
+	UK2Node_MakeMap* MakeMapNode = CompilerContext.SpawnIntermediateNode<UK2Node_MakeMap>(this, SourceGraph);
 	MakeMapNode->AllocateDefaultPins();
 	CompilerContext.MessageLog.NotifyIntermediateObjectCreation(MakeMapNode, this);
 
@@ -211,55 +290,58 @@ void UK2Node_MakeRequestHeader::ExpandNode(FKismetCompilerContext& CompilerConte
 	UEdGraphPin* CallFunctionInputPin = CallFunctionNode->FindPinChecked(TEXT("Headers"));
 	MapOutPin->MakeLinkTo(CallFunctionInputPin);
 	MakeMapNode->PinConnectionListChanged(MapOutPin);
-	const UEdGraphSchema_K2* Schema = GetDefault<UEdGraphSchema_K2>();
 
-	if (PresetEnumIndex == StaticCast<int32>(ERequestPresets::Custom))
+	const bool bIsUserSpecifiedHeader = PresetEnumIndex == StaticCast<int32>(ERequestPresets::Custom);
+	if (bIsUserSpecifiedHeader)
 	{
+		TArray<UEdGraphPin*> ValidKeyPins = UE::HttpBlueprint::Private::Pin::GetValidKeyPins(Pins);
+
 		// Add input pins to the MakeMap node
 		// UK2Node_MakeMap::AddInputPin adds two pins per call which is why we only need to call it for Pins.Num / 2
-		for (int32 Index{0}; Index < (Pins.Num() / 2); Index++)
+		for (int32 Index{0}; Index < ValidKeyPins.Num(); ++Index)
 		{
-			if (IsValidInputPin(Pins[Index]))
-			{
-				MakeMapNode->AddInputPin();
-			}
+			MakeMapNode->AddInputPin();
 		}
-	
-		// Skip the first index because it's the output pin
-		for (int32 Index{1}; Index < Pins.Num(); ++Index)
+
+		for (int32 Index{0}; Index < ValidKeyPins.Num(); ++Index)
 		{
-			if (IsValidInputPin(Pins[Index]))
-			{
-				UEdGraphPin* MapPin = MakeMapNode->FindPinChecked(MakeMapNode->GetPinName(Index - 2));
-				MapPin->PinType = Pins[Index]->PinType;
-				CompilerContext.MovePinLinksToIntermediate(*Pins[Index], *MapPin);
-			}
+			UEdGraphPin* KeyPin = ValidKeyPins[Index];
+			const int32 KeyPinIndex = GetPinIndex(KeyPin);
+
+			const int32 MapKeyPinIndex = Index * 2;
+
+			UEdGraphPin* MapKeyPin = MakeMapNode->FindPinChecked(MakeMapNode->GetPinName(MapKeyPinIndex));
+			MapKeyPin->PinType = KeyPin->PinType;
+			CompilerContext.MovePinLinksToIntermediate(*KeyPin, *MapKeyPin);
+
+			UEdGraphPin* ValuePin = GetPinAt(KeyPinIndex + 1);
+			
+			UEdGraphPin* MapValuePin = MakeMapNode->FindPinChecked(MakeMapNode->GetPinName(MapKeyPinIndex + 1));
+			MapValuePin->PinType = ValuePin->PinType;
+			CompilerContext.MovePinLinksToIntermediate(*ValuePin, *MapValuePin);
 		}
 	}
 	else
 	{
-		if (PresetEnumIndex != StaticCast<int32>(ERequestPresets::Custom))
+		// MakeMapNode spawns one set of pins by default. So we only need to Spawn NumPins - 1
+		for (int32 Index{0}; Index < UE::HttpBlueprint::Private::FPresetPair::NumPresetHeaders() - 1; Index++)
 		{
-			// MakeMapNode spawns one set of pins by default. So we only need to Spawn NumPins - 1
-			for (int32 Index{0}; Index < UE::HttpBlueprint::Private::FPresetPair::NumPresetHeaders() - 1; Index++)
-			{
-				MakeMapNode->AddInputPin();
-			}
-			
-			TArray<UEdGraphPin*> KeyPins, ValuePins;
-			MakeMapNode->GetKeyAndValuePins(KeyPins, ValuePins);
-			for (int32 Index{0}; Index < KeyPins.Num(); Index++)
-			{
-				UEdGraphPin* KeyPin = KeyPins[Index];
-				UEdGraphPin* ValuePin = ValuePins[Index];
+			MakeMapNode->AddInputPin();
+		}
 
-				Schema->TrySetDefaultValue(*KeyPin, UE::HttpBlueprint::Private::FPresetPair::GetHeaderKeyForIndex(Index));
-				Schema->TrySetDefaultValue(*ValuePin, UE::HttpBlueprint::Private::FPresetPair::GetHeaderValueForIndex(Index, PresetEnumIndex));
-			}
+		TArray<UEdGraphPin*> KeyPins, ValuePins;
+		MakeMapNode->GetKeyAndValuePins(KeyPins, ValuePins);
+		for (int32 Index{0}; Index < KeyPins.Num(); Index++)
+		{
+			UEdGraphPin* KeyPin = KeyPins[Index];
+			UEdGraphPin* ValuePin = ValuePins[Index];
+
+			Schema->TrySetDefaultValue(*KeyPin, UE::HttpBlueprint::Private::FPresetPair::GetHeaderKeyForIndex(Index));
+			Schema->TrySetDefaultValue(*ValuePin, UE::HttpBlueprint::Private::FPresetPair::GetHeaderValueForIndex(Index, PresetEnumIndex));
 		}
 	}
 
-	if (UEdGraphPin* OutputPin = FindPin(GetOutputPinName()))
+	if (UEdGraphPin* OutputPin = FindPin(UE::HttpBlueprint::Private::MakeRequestHeaderGlobals::OutputPinName))
 	{
 		if (UEdGraphPin* HeaderOutputPin = CallFunctionNode->FindPin(TEXT("OutHeader")))
 		{
@@ -297,21 +379,9 @@ bool UK2Node_MakeRequestHeader::IsNodePure() const
 	return true;
 }
 
-FName UK2Node_MakeRequestHeader::GetOutputPinName()
-{
-	return UE::HttpBlueprint::Private::MakeRequestHeaderGlobals::OutputPinName;
-}
-
 UEdGraphPin* UK2Node_MakeRequestHeader::GetOutputPin() const
 {
-	return FindPinChecked(GetOutputPinName());
-}
-
-bool UK2Node_MakeRequestHeader::IsValidInputPin(const UEdGraphPin* InputPin) const
-{
-	return InputPin->Direction == EGPD_Input
-		&& InputPin->PinName != UE::HttpBlueprint::Private::MakeRequestHeaderGlobals::PresetPinName
-		&& InputPin->PinType.PinCategory == UEdGraphSchema_K2::PC_String;
+	return FindPinChecked(UE::HttpBlueprint::Private::MakeRequestHeaderGlobals::OutputPinName);
 }
 
 void UK2Node_MakeRequestHeader::ConstructDefaultPinsForPreset()
@@ -328,14 +398,20 @@ void UK2Node_MakeRequestHeader::ConstructDefaultPinsForPreset()
 				const UEdGraphSchema_K2* Schema = GetDefault<UEdGraphSchema_K2>();
 				for (const FOptionalPin& OptionalPin : OptionalPins)
 				{
-					UEdGraphPin* NewPin = CreatePin(EGPD_Input, UEdGraphSchema_K2::PC_String, OptionalPin.PinName);
+					UEdGraphPin* ExistingOrNewPin = FindPin(OptionalPin.PinName);
+					if (!ExistingOrNewPin)
+					{
+						ExistingOrNewPin = CreatePin(EGPD_Input, UEdGraphSchema_K2::PC_String, OptionalPin.PinName);
+					}
+					
 					if (UEdGraphPin* LinkedToPin = OptionalPin.LinkedTo.Get())
 					{
-						NewPin->MakeLinkTo(LinkedToPin);						
-					}					
-					Schema->SetPinDefaultValueAtConstruction(NewPin, OptionalPin.PinDefaultValue);
+						ExistingOrNewPin->MakeLinkTo(LinkedToPin);						
+					}
+					
+					Schema->SetPinDefaultValueAtConstruction(ExistingOrNewPin, OptionalPin.PinDefaultValue);
 				}
-				NumInputs = Pins.Num() - 2;
+				NumInputs = (Pins.Num() - 2) / 2;
 				FBlueprintEditorUtils::MarkBlueprintAsStructurallyModified(GetBlueprint());
 			}
 		}
@@ -356,13 +432,7 @@ void UK2Node_MakeRequestHeader::ConstructDefaultPinsForPreset()
 				FScopedTransaction Transaction(LOCTEXT("RemovePinTx", "RemovePin"));
 				Modify();
 
-				FOptionalPin NewOptionalPin;
-				if (!PinToRemove->LinkedTo.IsEmpty())
-				{
-					NewOptionalPin.LinkedTo = PinToRemove->LinkedTo[0];
-				}
-				NewOptionalPin.PinName = PinToRemove->PinName;
-				NewOptionalPin.PinDefaultValue = PinToRemove->DefaultValue;
+				FOptionalPin NewOptionalPin = MakeOptionalPin(PinToRemove);
 				OptionalPins.Add(NewOptionalPin);
 				
 				TFunction<void(UEdGraphPin*)> RemovePinLambda = [this, &RemovePinLambda](UEdGraphPin* PinToRemove)->void
@@ -399,8 +469,9 @@ bool UK2Node_MakeRequestHeader::CanAddPin() const
 
 bool UK2Node_MakeRequestHeader::CanRemovePin(const UEdGraphPin* Pin) const
 {
-	return StaticCast<ERequestPresets>(PresetEnumIndex) == ERequestPresets::Custom && Pin->Direction == EGPD_Input &&
-		Pin->PinType.PinCategory == UEdGraphSchema_K2::PC_String;
+	return StaticCast<ERequestPresets>(PresetEnumIndex) == ERequestPresets::Custom
+		&& Pin->Direction == EGPD_Input
+		&& Pin->PinType.PinCategory == UEdGraphSchema_K2::PC_String;
 }
 
 void UK2Node_MakeRequestHeader::AddInputPin()
@@ -409,11 +480,13 @@ void UK2Node_MakeRequestHeader::AddInputPin()
 	Modify();
 
 	++NumInputs;
+	const FString KeyPinName = UE::HttpBlueprint::Private::Pin::MakePinName((NumInputs - 1) * 2);
 	CreatePin(EGPD_Input, UEdGraphSchema_K2::PC_String,
-	          *FString::Printf(TEXT("%s"), *GetPinName((NumInputs - 1) * 2).ToString()));
+	          *FString::Printf(TEXT("%s"), *KeyPinName));
 
+	const FString ValuePinName = UE::HttpBlueprint::Private::Pin::MakePinName((NumInputs - 1) * 2 + 1);
 	CreatePin(EGPD_Input, UEdGraphSchema_K2::PC_String,
-	          *FString::Printf(TEXT("%s"), *GetPinName((NumInputs - 1) * 2 + 1).ToString()));
+	          *FString::Printf(TEXT("%s"), *ValuePinName));
 	
 	if (!GetBlueprint()->bBeingCompiled)
 	{
@@ -424,74 +497,123 @@ void UK2Node_MakeRequestHeader::AddInputPin()
 
 void UK2Node_MakeRequestHeader::RemoveInputPin(UEdGraphPin* Pin)
 {
-	check(Pin);
-	check(Pin->Direction == EGPD_Input);
-	check(Pin->ParentPin == nullptr);
-	check(Pins.Contains(Pin));
-
 	FScopedTransaction Transaction(LOCTEXT("RemovePinTx", "RemovePin"));
 	Modify();
 
-	TFunction<void(UEdGraphPin*)> RemovePinLambda = [this, &RemovePinLambda](UEdGraphPin* PinToRemove)
+	ForEachInputPin(Pins, [this](UEdGraphPin* Pin, const int32& PinIndex)
+	{
+		Pins.RemoveAt(PinIndex);
+		Pin->MarkAsGarbage();
+		--NumInputs;
+	});
+
+	SyncPinNames();
+	FBlueprintEditorUtils::MarkBlueprintAsStructurallyModified(GetBlueprint());
+}
+
+void UK2Node_MakeRequestHeader::ForEachInputPin(
+	TArrayView<UEdGraphPin*> InPins,
+	TUniqueFunction<void(UEdGraphPin*, int32)>&& PinFunc)
+{
+	TFunction<bool(UEdGraphPin*)> ForEachPinLambda = [this, &ForEachPinLambda, &PinFunc](UEdGraphPin* PinToAffect)
 	{
-		check(PinToRemove);
+		check(PinToAffect);
 		
-		for (int32 SubPinIndex = PinToRemove->SubPins.Num()-1; SubPinIndex >= 0; --SubPinIndex)
+		for (int32 SubPinIndex = PinToAffect->SubPins.Num() - 1; SubPinIndex >= 0; --SubPinIndex)
 		{
-			RemovePinLambda(PinToRemove->SubPins[SubPinIndex]);
+			ForEachPinLambda(PinToAffect->SubPins[SubPinIndex]);
 		}
 
-		int32 PinRemovalIndex = INDEX_NONE;
-		if (Pins.Find(PinToRemove, PinRemovalIndex))
+		int32 PinIndex = INDEX_NONE;
+		if (Pins.Find(PinToAffect, PinIndex))
 		{
-			Pins.RemoveAt(PinRemovalIndex);
-			PinToRemove->MarkAsGarbage();
+			PinFunc(PinToAffect, PinIndex);
+			return true;
 		}
-	};
 
-	TArray<UEdGraphPin*> KeyPins;
-	TArray<UEdGraphPin*> ValuePins;
-	GetKeyAndValuePins(KeyPins, ValuePins);
+		return false;
+	};
 
-	int32 PinIndex = INDEX_NONE;
-	if (ValuePins.Find(Pin, PinIndex))
+	for (UEdGraphPin*& Pin : InPins)
 	{
-		RemovePinLambda(KeyPins[PinIndex]);
+		if (!Pin)
+		{
+			continue;
+		}
+
+		if (Pin->Direction != EEdGraphPinDirection::EGPD_Input)
+		{
+			continue;
+		}
+
+		if (!Pins.Contains(Pin))
+		{
+			UE_LOG(LogHttpBlueprintEditor, Warning, TEXT("Pin %s wasn't valid."), *Pin->GetName());
+			continue;
+		}
+
+		TArray<UEdGraphPin*> KeyPins;
+		TArray<UEdGraphPin*> ValuePins;
+		GetKeyAndValuePins(Pins, KeyPins, ValuePins);
+
+		int32 KVPPinIndex = INDEX_NONE;
+		
+		// If it's a value pin, then remove the key pin which necessarily exists
+		if (ValuePins.Find(Pin, KVPPinIndex))
+		{
+			ForEachPinLambda(KeyPins[KVPPinIndex]);
+		}
+		// Otherwise the inverse is true
+		else
+		{
+			verify(KeyPins.Find(Pin, KVPPinIndex));
+			ForEachPinLambda(ValuePins[KVPPinIndex]);
+		}
+
+		// What's left - either the key or value, depending on the above
+		if (ForEachPinLambda(Pin))
+		{
+			PinConnectionListChanged(Pin);
+		}
+		else
+		{
+			UE_LOG(LogHttpBlueprintEditor, Warning, TEXT("Pin %s wasn't affected."), *Pin->GetName());
+		}
 	}
-	else
+}
+
+FOptionalPin UK2Node_MakeRequestHeader::MakeOptionalPin(UEdGraphPin* InPinToCopy)
+{
+	FOptionalPin NewOptionalPin;
+	if (!InPinToCopy->LinkedTo.IsEmpty())
 	{
-		verify(KeyPins.Find(Pin, PinIndex));
-		RemovePinLambda(ValuePins[PinIndex]);
+		NewOptionalPin.LinkedTo = InPinToCopy->LinkedTo[0];
 	}
-
-	RemovePinLambda(Pin);
-	PinConnectionListChanged(Pin);
-
-	--NumInputs;
-	SyncPinNames();
-	FBlueprintEditorUtils::MarkBlueprintAsStructurallyModified(GetBlueprint());
+				
+	NewOptionalPin.PinName = InPinToCopy->PinName;
+	NewOptionalPin.PinDefaultValue = InPinToCopy->DefaultValue;
+	return NewOptionalPin;
 }
 
 void UK2Node_MakeRequestHeader::SyncPinNames()
 {
 	int32 CurrentNumParentPins = 0;
 	for (int32 PinIndex = 0; PinIndex < Pins.Num(); ++PinIndex)
-	{		
+	{
 		UEdGraphPin*& CurrentPin = Pins[PinIndex];
-		if (CurrentPin->Direction == EGPD_Input &&
-			CurrentPin->ParentPin == nullptr	&&
-			IsValidInputPin(Pins[PinIndex]))
+		if (CurrentPin->ParentPin == nullptr
+			&& UE::HttpBlueprint::Private::Pin::IsValidInputPin(Pins[PinIndex]))
 		{
 			const FName OldName = CurrentPin->PinName;
-			const FName ElementName = GetPinName(CurrentNumParentPins++);
+			const FString ElementName = UE::HttpBlueprint::Private::Pin::MakePinName(CurrentNumParentPins++);
 
 			CurrentPin->Modify();
-			CurrentPin->PinName = ElementName;
+			CurrentPin->PinName = FName(ElementName);
 
 			if (CurrentPin->SubPins.Num() > 0)
 			{
 				const FString OldNameStr = OldName.ToString();
-				const FString ElementNameStr = ElementName.ToString();
+				const FString ElementNameStr = ElementName;
 				FString OldFriendlyName = OldNameStr;
 				FString ElementFriendlyName = ElementNameStr;
 
@@ -512,36 +634,33 @@ void UK2Node_MakeRequestHeader::SyncPinNames()
 	}
 }
 
-FName UK2Node_MakeRequestHeader::GetPinName(const int32& PinIndex)
+void UK2Node_MakeRequestHeader::GetKeyAndValuePins(TArrayView<UEdGraphPin*> InPins, TArray<UEdGraphPin*>& KeyPins, TArray<UEdGraphPin*>& ValuePins) const
 {
-	const int32 PairIndex = PinIndex / 2;
-	if (PinIndex % 2 == 0)
-	{
-		return *FString::Printf(TEXT("Header %d"), PairIndex);
-	}
-
-	return *FString::Printf(TEXT("Value %d"), PairIndex);
-}
+	KeyPins.Reserve(InPins.Num());
+	ValuePins.Reserve(InPins.Num());
 
-void UK2Node_MakeRequestHeader::GetKeyAndValuePins(TArray<UEdGraphPin*>& KeyPins, TArray<UEdGraphPin*>& ValuePins) const
-{
+	// Store the last found key pin, such that PinFunc takes both the key and value pins
+	UEdGraphPin* LastPin = nullptr;
+		
 	// skip the first two (input, output) pins
-	for (int32 PinIndex = 2; PinIndex < Pins.Num(); ++PinIndex)
+	for (int32 PinIndex = 2; PinIndex < InPins.Num(); ++PinIndex)
 	{
-		UEdGraphPin* CurrentPin = Pins[PinIndex];
+		UEdGraphPin* CurrentPin = InPins[PinIndex];
 		if (CurrentPin->Direction == EGPD_Input && CurrentPin->ParentPin == nullptr)
 		{
-			// Key/Value pins alternate so if this is a map and the counts are even then this is a key
-			if (KeyPins.Num() == ValuePins.Num())
+			// Key/Value pins alternate so if the PinIndex is even, then this is a key
+			if (PinIndex % 2 == 0) 
 			{
-				KeyPins.Add(CurrentPin);
+				LastPin = CurrentPin;
 			}
 			else
 			{
+				KeyPins.Add(LastPin);
 				ValuePins.Add(CurrentPin);
 			}
 		}
 	}
+
 	check(KeyPins.Num() == ValuePins.Num());
 }
 
diff --git a/Engine/Plugins/Web/HttpBlueprint/Source/HttpBlueprintGraph/Private/K2Node_MakeRequestHeader.h b/Engine/Plugins/Web/HttpBlueprint/Source/HttpBlueprintGraph/Private/K2Node_MakeRequestHeader.h
--- a/Engine/Plugins/Web/HttpBlueprint/Source/HttpBlueprintGraph/Private/K2Node_MakeRequestHeader.h
+++ b/Engine/Plugins/Web/HttpBlueprint/Source/HttpBlueprintGraph/Private/K2Node_MakeRequestHeader.h
@@ -13,14 +13,17 @@ struct FEdGraphPinReference;
 class UGraphNodeContextMenuContext;
 class UToolMenu;
 
+/** Stores pin connections to restore after a node rebuild. */
 USTRUCT()
 struct FOptionalPin
 {
 	GENERATED_BODY()
 
+	/** Pin Name. */
 	UPROPERTY()
 	FName PinName;
 
+	/** Default value, stored as a string. */
 	UPROPERTY()
 	FString PinDefaultValue;
 
@@ -29,6 +32,7 @@ struct FOptionalPin
 	FEdGraphPinReference LinkedTo;
 };
 
+/** Node to create an Http header, with presets. */
 UCLASS(MinimalAPI)
 class UK2Node_MakeRequestHeader
 	: public UK2Node
@@ -40,6 +44,7 @@ public:
 	UK2Node_MakeRequestHeader(const FObjectInitializer& ObjectInitializer);
 
 	/** Begin UEdGraphNode Interface */
+	virtual void ReallocatePinsDuringReconstruction(TArray<UEdGraphPin*>& OldPins) override;
 	virtual void AllocateDefaultPins() override;
 	virtual FText GetNodeTitle(ENodeTitleType::Type TitleType) const override;
 	virtual FText GetTooltipText() const override;
@@ -62,28 +67,38 @@ public:
 	virtual void RemoveInputPin(UEdGraphPin* Pin) override;
 	/** End IK2Node_AddPinInterface */
 
-	static FName GetOutputPinName();
-
+	/** Returns the singular output pin containing the created http header. */
 	UEdGraphPin* GetOutputPin() const;
 
 protected:
-	UE_NODISCARD bool IsValidInputPin(const UEdGraphPin* InputPin) const;
-
+	/** Creates the default pin layout for the currently selected preset. */
 	void ConstructDefaultPinsForPreset();
 
+	/** Updates the pin names according to their type and input index. */
 	void SyncPinNames();
-	static FName GetPinName(const int32& PinIndex);
-	virtual void GetKeyAndValuePins(TArray<UEdGraphPin*>& KeyPins, TArray<UEdGraphPin*>& ValuePins) const;
+
+	/** Separates and returns the Key and Value pins corresponding to each input Key/Value pair. */
+	virtual void GetKeyAndValuePins(TArrayView<UEdGraphPin*> InPins, TArray<UEdGraphPin*>& KeyPins, TArray<UEdGraphPin*>& ValuePins) const;
+
+	/** Iterates over the provided Pins. */
+	void ForEachInputPin(TArrayView<UEdGraphPin*> InPins, TUniqueFunction<void(UEdGraphPin*,int32)>&& PinFunc);
+
+	/** Creates a copy of the input Pin as an Optional Pin. */
+	static FOptionalPin MakeOptionalPin(UEdGraphPin* InPinToCopy);
 
 private:
+	/** The UEnum type containing preset names. */
 	TObjectPtr<UEnum> PresetEnum;
 
+	/** The currently selected preset index. */
 	UPROPERTY()
 	int32 PresetEnumIndex;
 
+	/** Number of key/value inputs. */
 	UPROPERTY()
 	int32 NumInputs;
 
+	/** Used to store hidden or pins to be later restored. */
 	UPROPERTY()
 	TArray<FOptionalPin> OptionalPins;
 };
