PR #13143: [HttpBlueprint] Update HttpRequestProxyObject.cpp to fix crash on connection error

diff --git a/Engine/Plugins/Web/HttpBlueprint/Source/HttpBlueprint/Private/HttpRequestProxyObject.cpp b/Engine/Plugins/Web/HttpBlueprint/Source/HttpBlueprint/Private/HttpRequestProxyObject.cpp
--- a/Engine/Plugins/Web/HttpBlueprint/Source/HttpBlueprint/Private/HttpRequestProxyObject.cpp
+++ b/Engine/Plugins/Web/HttpBlueprint/Source/HttpBlueprint/Private/HttpRequestProxyObject.cpp
@@ -37,5 +37,8 @@ void UHttpRequestProxyObject::ProcessRequest(
 
 void UHttpRequestProxyObject::ProcessComplete(FHttpRequestPtr InRequest, FHttpResponsePtr InResponse, bool bInSuccessful)
 {
-	OnRequestComplete.Broadcast(InResponse->GetContentAsString(), bInSuccessful, MoveTemp(CachedHeader));
+	if (ensure(InRequest))
+	{
+		OnRequestComplete.Broadcast(InResponse ? InResponse->GetContentAsString() : LexToString(InRequest->GetFailureReason()), InResponse && bInSuccessful, MoveTemp(CachedHeader));
+	}
 }
