Fix blueprint compilation crash due to mismatching pin names

#10299

diff --git a/Engine/Plugins/Web/HttpBlueprint/Source/HttpBlueprintGraph/Private/K2Node_HttpRequest.cpp b/Engine/Plugins/Web/HttpBlueprint/Source/HttpBlueprintGraph/Private/K2Node_HttpRequest.cpp
--- a/Engine/Plugins/Web/HttpBlueprint/Source/HttpBlueprintGraph/Private/K2Node_HttpRequest.cpp
+++ b/Engine/Plugins/Web/HttpBlueprint/Source/HttpBlueprintGraph/Private/K2Node_HttpRequest.cpp
@@ -151,10 +151,10 @@ void UK2Node_HttpRequest::ExpandNode(FKismetCompilerContext& CompilerContext, UE
 		CallFromStructFunctionNode->AllocateDefaultPins();
 		CompilerContext.MessageLog.NotifyIntermediateObjectCreation(CallFromStructFunctionNode, this);
 
-		UEdGraphPin* InStructPin = CallFromStructFunctionNode->FindPinChecked(TEXT("InStruct"));
+		UEdGraphPin* InStructPin = CallFromStructFunctionNode->FindPinChecked(TEXT("Struct"));
 		InStructPin->PinType = GetBodyPin()->PinType;
 		
-		UEdGraphPin* OutJsonPin = CallFromStructFunctionNode->FindPinChecked(TEXT("OutJson"));
+		UEdGraphPin* OutJsonPin = CallFromStructFunctionNode->FindPinChecked(TEXT("OutJsonString"));
 
 		CompilerContext.MovePinLinksToIntermediate(*GetExecPin(), *CallFromStructFunctionNode->GetExecPin());
 		CompilerContext.MovePinLinksToIntermediate(*GetBodyPin(), *InStructPin);
